Refactor media player to use common  Java binding.

diff --git a/Engine/Plugins/Media/AndroidMedia/Source/AndroidMedia/Private/AndroidMediaPlayer.cpp b/Engine/Plugins/Media/AndroidMedia/Source/AndroidMedia/Private/AndroidMediaPlayer.cpp
--- a/Engine/Plugins/Media/AndroidMedia/Source/AndroidMedia/Private/AndroidMediaPlayer.cpp
+++ b/Engine/Plugins/Media/AndroidMedia/Source/AndroidMedia/Private/AndroidMediaPlayer.cpp
@@ -13,202 +13,6 @@ DEFINE_LOG_CATEGORY_STATIC(LogAndroidMediaPlayer, Log, All);
 
 #define LOCTEXT_NAMESPACE "FAndroidMediaModule"
 
-// Wrapper for com/epicgames/ue4/MediaPlayer.java.
-class FJavaAndroidMediaMediaPlayer : public FJavaClassObject
-{
-public:
-	FJavaAndroidMediaMediaPlayer()
-		: FJavaClassObject(GetClassName(), "()V")
-		, GetDurationMethod(GetClassMethod("getDuration", "()I"))
-		, ResetMethod(GetClassMethod("reset", "()V"))
-		, GetCurrentPositionMethod(GetClassMethod("getCurrentPosition", "()I"))
-		, IsLoopingMethod(GetClassMethod("isLooping", "()Z"))
-		, IsPlayingMethod(GetClassMethod("isPlaying", "()Z"))
-		, SetDataSourceURLMethod(GetClassMethod("setDataSource", "(Ljava/lang/String;)V"))
-		, SetDataSourceFileMethod(GetClassMethod("setDataSource", "(Ljava/lang/String;JJ)Z"))
-		, SetDataSourceAssetMethod(GetClassMethod("setDataSource", "(Landroid/content/res/AssetManager;Ljava/lang/String;JJ)Z"))
-		, PrepareMethod(GetClassMethod("prepare", "()V"))
-		, SeekToMethod(GetClassMethod("seekTo", "(I)V"))
-		, SetLoopingMethod(GetClassMethod("setLooping", "(Z)V"))
-		, ReleaseMethod(GetClassMethod("release", "()V"))
-		, GetVideoHeightMethod(GetClassMethod("getVideoHeight", "()I"))
-		, GetVideoWidthMethod(GetClassMethod("getVideoWidth", "()I"))
-		, SetVideoEnabledMethod(GetClassMethod("setVideoEnabled", "(Z)V"))
-		, SetAudioEnabledMethod(GetClassMethod("setAudioEnabled", "(Z)V"))
-		, GetVideoLastFrameDataMethod(GetClassMethod("getVideoLastFrameData", "()Ljava/nio/Buffer;"))
-		, StartMethod(GetClassMethod("start", "()V"))
-		, PauseMethod(GetClassMethod("pause", "()V"))
-		, StopMethod(GetClassMethod("stop", "()V"))
-	{
-	}
-
-	int32 GetDuration()
-	{
-		BindObjectToThread();
-		return CallMethod<int32>(GetDurationMethod);
-	}
-
-	void Reset()
-	{
-		BindObjectToThread();
-		CallMethod<void>(ResetMethod);
-	}
-
-	void Stop()
-	{
-		BindObjectToThread();
-		CallMethod<void>(StopMethod);
-	}
-
-	int32 GetCurrentPosition()
-	{
-		BindObjectToThread();
-		int32 position = CallMethod<int32>(GetCurrentPositionMethod);
-		return position;
-	}
-
-	bool IsLooping()
-	{
-		BindObjectToThread();
-		return CallMethod<bool>(IsLoopingMethod);
-	}
-
-	bool IsPlaying()
-	{
-		BindObjectToThread();
-		return CallMethod<bool>(IsPlayingMethod);
-	}
-
-	void SetDataSource(const FString & Url)
-	{
-		BindObjectToThread();
-		CallMethod<void>(SetDataSourceURLMethod, GetJString(Url));
-	}
-
-	bool SetDataSource(const FString& MoviePathOnDevice, int64 offset, int64 size)
-	{
-		BindObjectToThread();
-		return CallMethod<bool>(SetDataSourceFileMethod, GetJString(MoviePathOnDevice), offset, size);
-	}
-
-	bool SetDataSource(jobject AssetMgr, const FString& AssetPath, int64 offset, int64 size)
-	{
-		BindObjectToThread();
-		return CallMethod<bool>(SetDataSourceAssetMethod, AssetMgr, GetJString(AssetPath), offset, size);
-	}
-
-	void Prepare()
-	{
-		BindObjectToThread();
-		CallMethod<void>(PrepareMethod);
-	}
-
-	void SeekTo(int32 Milliseconds)
-	{
-		BindObjectToThread();
-		CallMethod<void>(SeekToMethod, Milliseconds);
-	}
-
-	void SetLooping(bool Looping)
-	{
-		BindObjectToThread();
-		CallMethod<void>(SetLoopingMethod, Looping);
-	}
-
-	void Release()
-	{
-		BindObjectToThread();
-		CallMethod<void>(ReleaseMethod);
-	}
-
-	int32 GetVideoHeight()
-	{
-		BindObjectToThread();
-		return CallMethod<int32>(GetVideoHeightMethod);
-	}
-
-	int32 GetVideoWidth()
-	{
-		BindObjectToThread();
-		return CallMethod<int32>(GetVideoWidthMethod);
-	}
-
-	void SetVideoEnabled(bool enabled = true)
-	{
-		BindObjectToThread();
-		CallMethod<void>(SetVideoEnabledMethod, enabled);
-	}
-
-	void SetAudioEnabled(bool enabled = true)
-	{
-		BindObjectToThread();
-		CallMethod<void>(SetAudioEnabledMethod, enabled);
-	}
-
-	bool GetVideoLastFrameData(void* & outPixels, int64 & outCount)
-	{
-		BindObjectToThread();
-		jobject buffer = CallMethod<jobject>(GetVideoLastFrameDataMethod);
-		if (nullptr != buffer)
-		{
-			outPixels = JEnv->GetDirectBufferAddress(buffer);
-			outCount = JEnv->GetDirectBufferCapacity(buffer);
-			JEnv->DeleteGlobalRef(buffer);
-		}
-		if (nullptr == buffer || nullptr == outPixels || 0 == outCount)
-		{
-			return false;
-		}
-		return true;
-	}
-
-	void Start()
-	{
-		BindObjectToThread();
-		CallMethod<void>(StartMethod);
-	}
-
-	void Pause()
-	{
-		BindObjectToThread();
-		CallMethod<void>(PauseMethod);
-	}
-
-private:
-	static FName GetClassName()
-	{
-		if (FAndroidMisc::GetAndroidBuildVersion() >= 14)
-		{
-			return FName("com/epicgames/ue4/MediaPlayer14");
-		}
-		else
-		{
-			return FName("");
-		}
-	}
-
-	FJavaClassMethod GetDurationMethod;
-	FJavaClassMethod ResetMethod;
-	FJavaClassMethod StopMethod;
-	FJavaClassMethod GetCurrentPositionMethod;
-	FJavaClassMethod IsLoopingMethod;
-	FJavaClassMethod IsPlayingMethod;
-	FJavaClassMethod SetDataSourceURLMethod;
-	FJavaClassMethod SetDataSourceFileMethod;
-	FJavaClassMethod SetDataSourceAssetMethod;
-	FJavaClassMethod PrepareMethod;
-	FJavaClassMethod SeekToMethod;
-	FJavaClassMethod SetLoopingMethod;
-	FJavaClassMethod ReleaseMethod;
-	FJavaClassMethod GetVideoHeightMethod;
-	FJavaClassMethod GetVideoWidthMethod;
-	FJavaClassMethod SetVideoEnabledMethod;
-	FJavaClassMethod SetAudioEnabledMethod;
-	FJavaClassMethod GetVideoLastFrameDataMethod;
-	FJavaClassMethod StartMethod;
-	FJavaClassMethod PauseMethod;
-};
-
 class FAndroidMediaPlayer::MediaTrack
 	: public IMediaTrack
 {
@@ -499,7 +303,7 @@ FAndroidMediaPlayer::FAndroidMediaPlayer()
 	: JavaMediaPlayer(nullptr)
 	, MediaState(EMediaState::Error)
 {
-	JavaMediaPlayer = MakeShareable(new FJavaAndroidMediaMediaPlayer());
+	JavaMediaPlayer = MakeShareable(new FJavaAndroidMediaPlayer());
 	if (JavaMediaPlayer.IsValid())
 	{
 		MediaState = EMediaState::Idle;
diff --git a/Engine/Plugins/Media/AndroidMedia/Source/AndroidMedia/Private/AndroidMediaPlayer.h b/Engine/Plugins/Media/AndroidMedia/Source/AndroidMedia/Private/AndroidMediaPlayer.h
--- a/Engine/Plugins/Media/AndroidMedia/Source/AndroidMedia/Private/AndroidMediaPlayer.h
+++ b/Engine/Plugins/Media/AndroidMedia/Source/AndroidMedia/Private/AndroidMediaPlayer.h
@@ -4,8 +4,7 @@
 
 #include "Media.h"
 #include "Tickable.h"
-
-class FJavaAndroidMediaMediaPlayer;
+#include "AndroidJavaMediaPlayer.h"
 
 /*
 	Implement media playback using the Android MediaPlayer interface.
@@ -91,7 +90,7 @@ private:
 	FOnMediaOpened OpenedEvent;
 
 	// The Java side media interface.
-	TSharedPtr<FJavaAndroidMediaMediaPlayer> JavaMediaPlayer;
+	TSharedPtr<FJavaAndroidMediaPlayer> JavaMediaPlayer;
 
 	// Currently opened media.
 	FString MediaUrl;
