Media: Removed deprecated file Media.h

diff --git a/Engine/Plugins/Media/AndroidMedia/Source/AndroidMedia/Private/AndroidMediaModule.cpp b/Engine/Plugins/Media/AndroidMedia/Source/AndroidMedia/Private/AndroidMediaModule.cpp
--- a/Engine/Plugins/Media/AndroidMedia/Source/AndroidMedia/Private/AndroidMediaModule.cpp
+++ b/Engine/Plugins/Media/AndroidMedia/Source/AndroidMedia/Private/AndroidMediaModule.cpp
@@ -1,11 +1,14 @@
 // Copyright 1998-2015 Epic Games, Inc. All Rights Reserved.
 
 #include "AndroidMediaPCH.h"
-
 #include "AndroidMediaPlayer.h"
+#include "IMediaModule.h"
+#include "IMediaPlayerFactory.h"
+
 
 #define LOCTEXT_NAMESPACE "FAndroidMediaModule"
 
+
 class FAndroidMediaModule
 	: public IModuleInterface
 	, public IMediaPlayerFactory
diff --git a/Engine/Plugins/Media/AndroidMedia/Source/AndroidMedia/Private/AndroidMediaPCH.h b/Engine/Plugins/Media/AndroidMedia/Source/AndroidMedia/Private/AndroidMediaPCH.h
--- a/Engine/Plugins/Media/AndroidMedia/Source/AndroidMedia/Private/AndroidMediaPCH.h
+++ b/Engine/Plugins/Media/AndroidMedia/Source/AndroidMedia/Private/AndroidMediaPCH.h
@@ -2,5 +2,12 @@
 
 #pragma once
 
+#include "Core.h"
+#include "IMediaInfo.h"
+#include "IMediaPlayer.h"
+#include "IMediaSink.h"
+#include "IMediaTrack.h"
+#include "IMediaTrackAudioDetails.h"
+#include "IMediaTrackCaptionDetails.h"
+#include "IMediaTrackVideoDetails.h"
 #include "ModuleManager.h"
-#include "Media.h"
diff --git a/Engine/Plugins/Media/AndroidMedia/Source/AndroidMedia/Private/AndroidMediaPlayer.h b/Engine/Plugins/Media/AndroidMedia/Source/AndroidMedia/Private/AndroidMediaPlayer.h
--- a/Engine/Plugins/Media/AndroidMedia/Source/AndroidMedia/Private/AndroidMediaPlayer.h
+++ b/Engine/Plugins/Media/AndroidMedia/Source/AndroidMedia/Private/AndroidMediaPlayer.h
@@ -2,10 +2,10 @@
 
 #pragma once
 
-#include "Media.h"
 #include "Tickable.h"
 #include "AndroidJavaMediaPlayer.h"
 
+
 /*
 	Implement media playback using the Android MediaPlayer interface.
 */
diff --git a/Engine/Source/Runtime/Media/Public/Deprecated/Media.h b/Engine/Source/Runtime/Media/Public/Deprecated/Media.h
deleted file mode 100644
