Media: Removed unused dependencies from player plug-ins

diff --git a/Engine/Plugins/Media/AndroidMedia/Source/AndroidMedia/AndroidMedia.Build.cs b/Engine/Plugins/Media/AndroidMedia/Source/AndroidMedia/AndroidMedia.Build.cs
--- a/Engine/Plugins/Media/AndroidMedia/Source/AndroidMedia/AndroidMedia.Build.cs
+++ b/Engine/Plugins/Media/AndroidMedia/Source/AndroidMedia/AndroidMedia.Build.cs
@@ -9,13 +9,11 @@ namespace UnrealBuildTool.Rules
             DynamicallyLoadedModuleNames.AddRange(
                 new string[] {
                     "Media",
-                    "Settings",
 				});
 
 			PrivateDependencyModuleNames.AddRange(
 				new string[] {
 					"Core",
-				    "CoreUObject",
 				    "Engine",
                     "RenderCore",
 				});
@@ -23,7 +21,6 @@ namespace UnrealBuildTool.Rules
 			PrivateIncludePathModuleNames.AddRange(
 				new string[] {
                     "Media",
-					"Settings",
 				});
 
 			PrivateIncludePaths.AddRange(
diff --git a/Engine/Plugins/Media/AvfMedia/Source/AvfMedia/AvfMedia.Build.cs b/Engine/Plugins/Media/AvfMedia/Source/AvfMedia/AvfMedia.Build.cs
--- a/Engine/Plugins/Media/AvfMedia/Source/AvfMedia/AvfMedia.Build.cs
+++ b/Engine/Plugins/Media/AvfMedia/Source/AvfMedia/AvfMedia.Build.cs
@@ -9,14 +9,12 @@ namespace UnrealBuildTool.Rules
 			DynamicallyLoadedModuleNames.AddRange(
 				new string[] {
 					"Media",
-					"Settings",
 				}
 			);
 
 			PrivateDependencyModuleNames.AddRange(
 				new string[] {
 					"Core",
-					"CoreUObject",
 					"RenderCore",
 				}
 			);
@@ -24,7 +22,6 @@ namespace UnrealBuildTool.Rules
 			PrivateIncludePathModuleNames.AddRange(
 				new string[] {
 					"Media",
-					"Settings",
 				}
 			);
 
diff --git a/Engine/Plugins/Media/WmfMedia/Source/WmfMedia/Private/WmfMediaPrivatePCH.h b/Engine/Plugins/Media/WmfMedia/Source/WmfMedia/Private/WmfMediaPrivatePCH.h
--- a/Engine/Plugins/Media/WmfMedia/Source/WmfMedia/Private/WmfMediaPrivatePCH.h
+++ b/Engine/Plugins/Media/WmfMedia/Source/WmfMedia/Private/WmfMediaPrivatePCH.h
@@ -7,7 +7,6 @@
  *****************************************************************************/
 
 #include "Core.h"
-#include "CoreUObject.h"
 
 #include "IMediaInfo.h"
 #include "IMediaPlayer.h"
diff --git a/Engine/Plugins/Media/WmfMedia/Source/WmfMedia/WmfMedia.Build.cs b/Engine/Plugins/Media/WmfMedia/Source/WmfMedia/WmfMedia.Build.cs
--- a/Engine/Plugins/Media/WmfMedia/Source/WmfMedia/WmfMedia.Build.cs
+++ b/Engine/Plugins/Media/WmfMedia/Source/WmfMedia/WmfMedia.Build.cs
@@ -9,14 +9,12 @@ namespace UnrealBuildTool.Rules
             DynamicallyLoadedModuleNames.AddRange(
                 new string[] {
                     "Media",
-                    "Settings",
 				}
             );
 
 			PrivateDependencyModuleNames.AddRange(
 				new string[] {
 					"Core",
-					"CoreUObject",
                     "RenderCore",
 				}
 			);
@@ -24,7 +22,6 @@ namespace UnrealBuildTool.Rules
 			PrivateIncludePathModuleNames.AddRange(
 				new string[] {
                     "Media",
-					"Settings",
 				}
 			);
 
