Plugins: Separate categories and better names for media framework plug-ins

diff --git a/Engine/Plugins/Media/AndroidMedia/AndroidMedia.uplugin b/Engine/Plugins/Media/AndroidMedia/AndroidMedia.uplugin
--- a/Engine/Plugins/Media/AndroidMedia/AndroidMedia.uplugin
+++ b/Engine/Plugins/Media/AndroidMedia/AndroidMedia.uplugin
@@ -1,7 +1,7 @@
 {
     "FileVersion" : 1,
 	
-	"FriendlyName" : "Android Media",
+	"FriendlyName" : "Android Media Player",
 	"Version" : 1,
 	"VersionName" : "1.0",
 	"CreatedBy" : "Epic Games, Inc.",
diff --git a/Engine/Plugins/Media/AvfMedia/AvfMedia.uplugin b/Engine/Plugins/Media/AvfMedia/AvfMedia.uplugin
--- a/Engine/Plugins/Media/AvfMedia/AvfMedia.uplugin
+++ b/Engine/Plugins/Media/AvfMedia/AvfMedia.uplugin
@@ -1,7 +1,7 @@
 {
 	"FileVersion" : 3,
 	
-	"FriendlyName" : "Avf Video",
+	"FriendlyName" : "AVF Media Player",
 	"Version" : 1,
 	"VersionName" : "1.0",
 	"CreatedBy" : "Epic Games, Inc.",
diff --git a/Engine/Plugins/Media/WmfMedia/WmfMedia.uplugin b/Engine/Plugins/Media/WmfMedia/WmfMedia.uplugin
--- a/Engine/Plugins/Media/WmfMedia/WmfMedia.uplugin
+++ b/Engine/Plugins/Media/WmfMedia/WmfMedia.uplugin
@@ -1,7 +1,7 @@
 {
 	"FileVersion" : 3,
 	
-	"FriendlyName" : "Windows Media Foundation Video",
+	"FriendlyName" : "WMF Media Player",
 	"Version" : 1,
 	"VersionName" : "1.0",
 	"CreatedBy" : "Epic Games, Inc.",
