More fixes for non-unity compile.

diff --git a/Engine/Plugins/Media/AndroidMedia/Source/AndroidMediaEditor/Private/AndroidMediaEditorModule.cpp b/Engine/Plugins/Media/AndroidMedia/Source/AndroidMediaEditor/Private/AndroidMediaEditorModule.cpp
--- a/Engine/Plugins/Media/AndroidMedia/Source/AndroidMediaEditor/Private/AndroidMediaEditorModule.cpp
+++ b/Engine/Plugins/Media/AndroidMedia/Source/AndroidMediaEditor/Private/AndroidMediaEditorModule.cpp
@@ -1,6 +1,8 @@
 // Copyright 1998-2016 Epic Games, Inc. All Rights Reserved.
 
 #include "AndroidMediaEditorPCH.h"
+#include "CoreUObject.h"
+#include "ModuleManager.h"
 #include "ModuleInterface.h"
 
 
diff --git a/Engine/Plugins/Media/AvfMedia/Source/AvfMediaEditor/Private/AvfMediaEditorModule.cpp b/Engine/Plugins/Media/AvfMedia/Source/AvfMediaEditor/Private/AvfMediaEditorModule.cpp
--- a/Engine/Plugins/Media/AvfMedia/Source/AvfMediaEditor/Private/AvfMediaEditorModule.cpp
+++ b/Engine/Plugins/Media/AvfMedia/Source/AvfMediaEditor/Private/AvfMediaEditorModule.cpp
@@ -1,6 +1,8 @@
 // Copyright 1998-2016 Epic Games, Inc. All Rights Reserved.
 
 #include "AvfMediaEditorPCH.h"
+#include "CoreUObject.h"
+#include "ModuleManager.h"
 #include "ModuleInterface.h"
 
 
diff --git a/Engine/Plugins/Media/WmfMedia/Source/WmfMediaEditor/Private/WmfMediaEditorModule.cpp b/Engine/Plugins/Media/WmfMedia/Source/WmfMediaEditor/Private/WmfMediaEditorModule.cpp
--- a/Engine/Plugins/Media/WmfMedia/Source/WmfMediaEditor/Private/WmfMediaEditorModule.cpp
+++ b/Engine/Plugins/Media/WmfMedia/Source/WmfMediaEditor/Private/WmfMediaEditorModule.cpp
@@ -1,6 +1,8 @@
 // Copyright 1998-2016 Epic Games, Inc. All Rights Reserved.
 
 #include "WmfMediaEditorPCH.h"
+#include "CoreUObject.h"
+#include "ModuleManager.h"
 #include "ModuleInterface.h"
 
 
