Fixed missing includes

diff --git a/Engine/Plugins/Media/AndroidMedia/Source/AndroidMediaEditor/Private/AndroidMediaEditorModule.cpp b/Engine/Plugins/Media/AndroidMedia/Source/AndroidMediaEditor/Private/AndroidMediaEditorModule.cpp
--- a/Engine/Plugins/Media/AndroidMedia/Source/AndroidMediaEditor/Private/AndroidMediaEditorModule.cpp
+++ b/Engine/Plugins/Media/AndroidMedia/Source/AndroidMediaEditor/Private/AndroidMediaEditorModule.cpp
@@ -1,8 +1,6 @@
 // Copyright 1998-2016 Epic Games, Inc. All Rights Reserved.
 
 #include "AndroidMediaEditorPCH.h"
-#include "CoreUObject.h"
-#include "ModuleManager.h"
 #include "ModuleInterface.h"
 
 
diff --git a/Engine/Plugins/Media/AndroidMedia/Source/AndroidMediaEditor/Private/AndroidMediaEditorPCH.h b/Engine/Plugins/Media/AndroidMedia/Source/AndroidMediaEditor/Private/AndroidMediaEditorPCH.h
--- a/Engine/Plugins/Media/AndroidMedia/Source/AndroidMediaEditor/Private/AndroidMediaEditorPCH.h
+++ b/Engine/Plugins/Media/AndroidMedia/Source/AndroidMediaEditor/Private/AndroidMediaEditorPCH.h
@@ -3,4 +3,6 @@
 #pragma once
 
 #include "Runtime/Core/Public/Core.h"
+#include "Runtime/Core/Public/Modules/ModuleManager.h"
+#include "Runtime/CoreUObject/Public/CoreUObject.h"
 #include "Runtime/MediaAssets/Public/FileMediaSource.h"
diff --git a/Engine/Plugins/Media/AvfMedia/Source/AvfMediaEditor/Private/AvfMediaEditorModule.cpp b/Engine/Plugins/Media/AvfMedia/Source/AvfMediaEditor/Private/AvfMediaEditorModule.cpp
--- a/Engine/Plugins/Media/AvfMedia/Source/AvfMediaEditor/Private/AvfMediaEditorModule.cpp
+++ b/Engine/Plugins/Media/AvfMedia/Source/AvfMediaEditor/Private/AvfMediaEditorModule.cpp
@@ -1,8 +1,6 @@
 // Copyright 1998-2016 Epic Games, Inc. All Rights Reserved.
 
 #include "AvfMediaEditorPCH.h"
-#include "CoreUObject.h"
-#include "ModuleManager.h"
 #include "ModuleInterface.h"
 
 
diff --git a/Engine/Plugins/Media/AvfMedia/Source/AvfMediaEditor/Private/AvfMediaEditorPCH.h b/Engine/Plugins/Media/AvfMedia/Source/AvfMediaEditor/Private/AvfMediaEditorPCH.h
--- a/Engine/Plugins/Media/AvfMedia/Source/AvfMediaEditor/Private/AvfMediaEditorPCH.h
+++ b/Engine/Plugins/Media/AvfMedia/Source/AvfMediaEditor/Private/AvfMediaEditorPCH.h
@@ -3,4 +3,6 @@
 #pragma once
 
 #include "Runtime/Core/Public/Core.h"
+#include "Runtime/Core/Public/Modules/ModuleManager.h"
+#include "Runtime/CoreUObject/Public/CoreUObject.h"
 #include "Runtime/MediaAssets/Public/FileMediaSource.h"
diff --git a/Engine/Plugins/Media/WmfMedia/Source/WmfMediaEditor/Private/WmfMediaEditorModule.cpp b/Engine/Plugins/Media/WmfMedia/Source/WmfMediaEditor/Private/WmfMediaEditorModule.cpp
--- a/Engine/Plugins/Media/WmfMedia/Source/WmfMediaEditor/Private/WmfMediaEditorModule.cpp
+++ b/Engine/Plugins/Media/WmfMedia/Source/WmfMediaEditor/Private/WmfMediaEditorModule.cpp
@@ -1,8 +1,6 @@
 // Copyright 1998-2016 Epic Games, Inc. All Rights Reserved.
 
 #include "WmfMediaEditorPCH.h"
-#include "CoreUObject.h"
-#include "ModuleManager.h"
 #include "ModuleInterface.h"
 
 
diff --git a/Engine/Plugins/Media/WmfMedia/Source/WmfMediaEditor/Private/WmfMediaEditorPCH.h b/Engine/Plugins/Media/WmfMedia/Source/WmfMediaEditor/Private/WmfMediaEditorPCH.h
--- a/Engine/Plugins/Media/WmfMedia/Source/WmfMediaEditor/Private/WmfMediaEditorPCH.h
+++ b/Engine/Plugins/Media/WmfMedia/Source/WmfMediaEditor/Private/WmfMediaEditorPCH.h
@@ -3,4 +3,6 @@
 #pragma once
 
 #include "Runtime/Core/Public/Core.h"
+#include "Runtime/Core/Public/Modules/ModuleManager.h"
+#include "Runtime/CoreUObject/Public/CoreUObject.h"
 #include "Runtime/MediaAssets/Public/FileMediaSource.h"
