Copying //UE4/Dev-Sequencer to //UE4/Main (Source: //UE4/Dev-Sequencer @ 3617642)

diff --git a/Engine/Plugins/Media/AndroidMedia/Source/AndroidMedia/Private/Player/AndroidMediaPlayer.cpp b/Engine/Plugins/Media/AndroidMedia/Source/AndroidMedia/Private/Player/AndroidMediaPlayer.cpp
--- a/Engine/Plugins/Media/AndroidMedia/Source/AndroidMedia/Private/Player/AndroidMediaPlayer.cpp
+++ b/Engine/Plugins/Media/AndroidMedia/Source/AndroidMedia/Private/Player/AndroidMediaPlayer.cpp
@@ -1,23 +1,54 @@
 // Copyright 1998-2017 Epic Games, Inc. All Rights Reserved.
 
 #include "AndroidMediaPlayer.h"
-#include "AndroidJavaMediaPlayer.h"
+#include "AndroidMediaPrivate.h"
+
+#include "Android/AndroidFile.h"
+#include "AndroidMediaSettings.h"
+#include "IMediaEventSink.h"
+#include "MediaSamples.h"
 #include "Misc/CoreDelegates.h"
+#include "Misc/Guid.h"
 #include "Misc/Paths.h"
-#include "Android/AndroidFile.h"
-#include "../AndroidMediaLog.h"
+#include "RenderingThread.h"
+#include "RHI.h"
+#include "RHIResources.h"
+#include "UObject/Class.h"
+#include "UObject/UObjectGlobals.h"
+#include "ExternalTexture.h"
 
-#define LOCTEXT_NAMESPACE "FAndroidMediaModule"
+#include "AndroidJavaMediaPlayer.h"
+#include "AndroidMediaTextureSample.h"
+
+#define ANDROIDMEDIAPLAYER_USE_EXTERNALTEXTURE 1
+#define ANDROIDMEDIAPLAYER_USE_PREPAREASYNC 1
+#define ANDROIDMEDIAPLAYER_USE_NATIVELOGGING 1
+
+#define LOCTEXT_NAMESPACE "FAndroidMediaPlayer"
 
 
 /* FAndroidMediaPlayer structors
  *****************************************************************************/
 
-FAndroidMediaPlayer::FAndroidMediaPlayer()
-	: JavaMediaPlayer(nullptr)
+FAndroidMediaPlayer::FAndroidMediaPlayer(IMediaEventSink& InEventSink)
+	: CurrentState(EMediaState::Closed)
+	, bLooping(false)
+	, EventSink(InEventSink)
+#if WITH_ENGINE
+	, JavaMediaPlayer(MakeShared<FJavaAndroidMediaPlayer, ESPMode::ThreadSafe>(false, FAndroidMisc::ShouldUseVulkan()))
+#else
+	, JavaMediaPlayer(MakeShared<FJavaAndroidMediaPlayer, ESPMode::ThreadSafe>(true, FAndroidMisc::ShouldUseVulkan()))
+#endif
+	, Samples(MakeShared<FMediaSamples, ESPMode::ThreadSafe>())
+	, SelectedAudioTrack(INDEX_NONE)
+	, SelectedCaptionTrack(INDEX_NONE)
+	, SelectedVideoTrack(INDEX_NONE)
+	, VideoSamplePool(new FAndroidMediaTextureSamplePool)
 {
-	JavaMediaPlayer = MakeShareable(new FJavaAndroidMediaPlayer(FAndroidMisc::ShouldUseVulkan()));
-	State = JavaMediaPlayer.IsValid() ? EMediaState::Closed : EMediaState::Error;
+	check(JavaMediaPlayer.IsValid());
+	check(Samples.IsValid());
+
+	CurrentState = (JavaMediaPlayer.IsValid() && Samples.IsValid())? EMediaState::Closed : EMediaState::Error;
 }
 
 
@@ -27,159 +58,46 @@ FAndroidMediaPlayer::~FAndroidMediaPlayer()
 
 	if (JavaMediaPlayer.IsValid())
 	{
-		JavaMediaPlayer->Reset();
-		JavaMediaPlayer->Release();
-	}
-}
-
-
-void FAndroidMediaPlayer::AppServicePause()
-{
-	// check state in case changed before ticked
-	if ((State == EMediaState::Playing) && JavaMediaPlayer.IsValid())
-	{
-		JavaMediaPlayer->Pause();
-	}
-}
-
-
-void FAndroidMediaPlayer::AppServiceResume()
-{
-	// check state in case changed before ticked
-	if ((State == EMediaState::Playing) && JavaMediaPlayer.IsValid())
-	{
-		JavaMediaPlayer->Start();
-	}
-}
-
-
-/* IMediaControls interface
- *****************************************************************************/
-
-FTimespan FAndroidMediaPlayer::GetDuration() const
-{
-	if (State == EMediaState::Error)
-	{
-		return FTimespan::Zero();
-	}
-
-	return FTimespan::FromMilliseconds(JavaMediaPlayer->GetDuration());
-}
-
-
-float FAndroidMediaPlayer::GetRate() const
-{
-	return (State == EMediaState::Playing) ? 1.0f : 0.0f;
-}
-
-
-EMediaState FAndroidMediaPlayer::GetState() const
-{
-	return State;
-}
-
-
-TRange<float> FAndroidMediaPlayer::GetSupportedRates(EMediaPlaybackDirections Direction, bool Unthinned) const
-{
-	if (Direction == EMediaPlaybackDirections::Reverse)
-	{
-		return TRange<float>::Empty();
-	}
-
-	return TRange<float>(1.0f);
-}
-
-
-FTimespan FAndroidMediaPlayer::GetTime() const
-{
-	if (State == EMediaState::Error)
-	{
-		return FTimespan::Zero();
-	}
-
-	return FTimespan::FromMilliseconds(JavaMediaPlayer->GetCurrentPosition());
-}
-
-
-bool FAndroidMediaPlayer::IsLooping() const
-{
-	return JavaMediaPlayer.IsValid() ? JavaMediaPlayer->IsLooping() : false;
-}
-
-
-bool FAndroidMediaPlayer::Seek(const FTimespan& Time)
-{
-	if ((State == EMediaState::Closed) || (State == EMediaState::Error))
-	{
-		return false;
-	}
-
-	JavaMediaPlayer->SeekTo(static_cast<int32>(Time.GetTotalMilliseconds()));
-
-	return true;
-}
+		if (GSupportsImageExternal && !FAndroidMisc::ShouldUseVulkan())
+		{
+			// Unregister the external texture on render thread
+			FTextureRHIRef VideoTexture = JavaMediaPlayer->GetVideoTexture();
 
+			JavaMediaPlayer->SetVideoTexture(nullptr);
+			JavaMediaPlayer->Reset();
+			JavaMediaPlayer->Release();
 
-bool FAndroidMediaPlayer::SetLooping(bool Looping)
-{
-	if (!JavaMediaPlayer.IsValid())
-	{
-		return false;
-	}
+			struct FReleaseVideoResourcesParams
+			{
+				FTextureRHIRef VideoTexture;
+				FGuid PlayerGuid;
+			};
 
-	JavaMediaPlayer->SetLooping(Looping);
+			FReleaseVideoResourcesParams ReleaseVideoResourcesParams = { VideoTexture, PlayerGuid };
 
-	return true;
-}
+			ENQUEUE_UNIQUE_RENDER_COMMAND_ONEPARAMETER(AndroidMediaPlayerWriteVideoSample,
+				FReleaseVideoResourcesParams, Params, ReleaseVideoResourcesParams,
+				{
+					#if ANDROIDMEDIAPLAYER_USE_NATIVELOGGING
+						FPlatformMisc::LowLevelOutputDebugStringf(TEXT("~FAndroidMediaPlayer: Unregister Guid: %s"), *Params.PlayerGuid.ToString());
+					#endif
 
+					FExternalTextureRegistry::Get().UnregisterExternalTexture(Params.PlayerGuid);
 
-bool FAndroidMediaPlayer::SetRate(float Rate)
-{
-	switch (State)
-	{
-	case EMediaState::Playing:
-		if (FMath::IsNearlyZero(Rate))
-		{
-			JavaMediaPlayer->Pause();
-			State = EMediaState::Paused;
-			MediaEvent.Broadcast(EMediaEvent::PlaybackSuspended);
-			
-			return true;
+					// @todo: this causes a crash
+//					Params.VideoTexture->Release();
+				});
 		}
-		break;
-
-	case EMediaState::Paused:
-	case EMediaState::Stopped:
-		if (Rate > 0.0f)
+		else
 		{
-			JavaMediaPlayer->Start();
-			State = EMediaState::Playing;
-			MediaEvent.Broadcast(EMediaEvent::PlaybackResumed);
-			
-			return true;
+			JavaMediaPlayer->SetVideoTexture(nullptr);
+			JavaMediaPlayer->Reset();
+			JavaMediaPlayer->Release();
 		}
-		break;
 	}
 
-	return false;
-}
-
-
-bool FAndroidMediaPlayer::SupportsRate(float Rate, bool Unthinned) const
-{
-	return (Rate == 1.0f);
-}
-
-
-bool FAndroidMediaPlayer::SupportsScrubbing() const
-{
-	return true;
-}
-
-
-bool FAndroidMediaPlayer::SupportsSeeking() const
-{
-	return true;
+	delete VideoSamplePool;
+	VideoSamplePool = nullptr;
 }
 
 
@@ -188,38 +106,49 @@ bool FAndroidMediaPlayer::SupportsSeeking() const
 
 void FAndroidMediaPlayer::Close()
 {
-	// remove delegates if registered
-	if (ResumeHandle.IsValid())
-	{
-		FCoreDelegates::ApplicationHasEnteredForegroundDelegate.Remove(ResumeHandle);
-		ResumeHandle.Reset();
-	}
-	if (PauseHandle.IsValid())
-	{
-		FCoreDelegates::ApplicationWillEnterBackgroundDelegate.Remove(PauseHandle);
-		PauseHandle.Reset();
-	}
+	#if ANDROIDMEDIAPLAYER_USE_NATIVELOGGING
+		FPlatformMisc::LowLevelOutputDebugStringf(TEXT("FAndroidMedia::Close() - %s"), *PlayerGuid.ToString());
+	#endif
 
-	if (State == EMediaState::Closed)
+	if (CurrentState == EMediaState::Closed)
 	{
 		return;
 	}
 
-	Tracks.Reset();
-	Info.Empty();
+	CurrentState = EMediaState::Closed;
+
+	bLooping = false;
 
 	if (JavaMediaPlayer.IsValid())
 	{
+		JavaMediaPlayer->SetLooping(false);
+//		JavaMediaPlayer->SetVideoTextureValid(false);
 		JavaMediaPlayer->Stop();
 		JavaMediaPlayer->Reset();
 	}
 
+	VideoSamplePool->Reset();
+
+	SelectedAudioTrack = INDEX_NONE;
+	SelectedCaptionTrack = INDEX_NONE;
+	SelectedVideoTrack = INDEX_NONE;
+
+	AudioTracks.Empty();
+	CaptionTracks.Empty();
+	VideoTracks.Empty();
+
+	Info.Empty();
 	MediaUrl = FString();
-	State = EMediaState::Closed;
 
 	// notify listeners
-	MediaEvent.Broadcast(EMediaEvent::TracksChanged);
-	MediaEvent.Broadcast(EMediaEvent::MediaClosed);
+	EventSink.ReceiveMediaEvent(EMediaEvent::TracksChanged);
+	EventSink.ReceiveMediaEvent(EMediaEvent::MediaClosed);
+}
+
+
+IMediaCache& FAndroidMediaPlayer::GetCache()
+{
+	return *this;
 }
 
 
@@ -242,9 +171,9 @@ FName FAndroidMediaPlayer::GetName() const
 }
 
 
-IMediaOutput& FAndroidMediaPlayer::GetOutput()
+IMediaSamples& FAndroidMediaPlayer::GetSamples()
 {
-	return Tracks;
+	return *Samples.Get();
 }
 
 
@@ -256,7 +185,7 @@ FString FAndroidMediaPlayer::GetStats() const
 
 IMediaTracks& FAndroidMediaPlayer::GetTracks()
 {
-	return Tracks;
+	return *this;
 }
 
 
@@ -266,9 +195,19 @@ FString FAndroidMediaPlayer::GetUrl() const
 }
 
 
-bool FAndroidMediaPlayer::Open(const FString& Url, const IMediaOptions& Options)
+IMediaView& FAndroidMediaPlayer::GetView()
+{
+	return *this;
+}
+
+
+bool FAndroidMediaPlayer::Open(const FString& Url, const IMediaOptions* /*Options*/)
 {
-	if (State == EMediaState::Error)
+	#if ANDROIDMEDIAPLAYER_USE_NATIVELOGGING
+		FPlatformMisc::LowLevelOutputDebugStringf(TEXT("FAndroidMedia::Open(%s) - %s"), *Url, *PlayerGuid.ToString());
+	#endif
+
+	if (CurrentState == EMediaState::Error)
 	{
 		return false;
 	}
@@ -330,36 +269,319 @@ bool FAndroidMediaPlayer::Open(const FString& Url, const IMediaOptions& Options)
 	// prepare media
 	MediaUrl = Url;
 
+#if ANDROIDMEDIAPLAYER_USE_PREPAREASYNC
+	if (!JavaMediaPlayer->PrepareAsync())
+	{
+		UE_LOG(LogAndroidMedia, Warning, TEXT("Failed to prepare media source %s"), *Url);
+		return false;
+	}
+
+	CurrentState = EMediaState::Preparing;
+	return true;
+#else
 	if (!JavaMediaPlayer->Prepare())
 	{
 		UE_LOG(LogAndroidMedia, Warning, TEXT("Failed to prepare media source %s"), *Url);
 		return false;
 	}
 
-	State = EMediaState::Stopped;
+	return InitializePlayer();
+#endif
+}
 
-	Tracks.Initialize(JavaMediaPlayer.ToSharedRef(), Info);
-	Tracks.SelectTrack(EMediaTrackType::Audio, 0);
-	Tracks.SelectTrack(EMediaTrackType::Video, 0);
 
-	// notify listeners
-	MediaEvent.Broadcast(EMediaEvent::TracksChanged);
-	MediaEvent.Broadcast(EMediaEvent::MediaOpened);
+bool FAndroidMediaPlayer::Open(const TSharedRef<FArchive, ESPMode::ThreadSafe>& /*Archive*/, const FString& /*OriginalUrl*/, const IMediaOptions* /*Options*/)
+{
+	return false; // @todo AndroidMedia: implement opening media from FArchive
+}
 
-	return true;
+
+void FAndroidMediaPlayer::SetGuid(const FGuid& Guid)
+{
+	PlayerGuid = Guid;
+
+	#if ANDROIDMEDIAPLAYER_USE_NATIVELOGGING
+		FPlatformMisc::LowLevelOutputDebugStringf(TEXT("IMediaPlayer SetGuid: %s"), *PlayerGuid.ToString());
+	#endif
 }
 
 
-bool FAndroidMediaPlayer::Open(const TSharedRef<FArchive, ESPMode::ThreadSafe>& Archive, const FString& OriginalUrl, const IMediaOptions& Options)
+void FAndroidMediaPlayer::TickFetch(FTimespan DeltaTime)
 {
-	// @todo AndroidMedia: implement opening media from FArchive
-	return false;
+	if (CurrentState != EMediaState::Playing && CurrentState != EMediaState::Paused)
+	{
+		return;
+	}
+
+	if (!JavaMediaPlayer.IsValid() || !VideoTracks.IsValidIndex(SelectedVideoTrack))
+	{
+		return;
+	}
+
+	// deal with resolution changes (usually from streams)
+	if (JavaMediaPlayer->DidResolutionChange())
+	{
+		JavaMediaPlayer->SetVideoTextureValid(false);
+
+		// The video track dimensions need updating
+		FIntPoint Dimensions = FIntPoint(JavaMediaPlayer->GetVideoWidth(), JavaMediaPlayer->GetVideoHeight());
+		VideoTracks[SelectedVideoTrack].Dimensions = Dimensions;
+	}
+
+#if WITH_ENGINE
+
+	if (FAndroidMisc::ShouldUseVulkan())
+	{
+		// create new video sample
+		const FJavaAndroidMediaPlayer::FVideoTrack VideoTrack = VideoTracks[SelectedVideoTrack];
+		auto VideoSample = VideoSamplePool->AcquireShared();
+
+		if (!VideoSample->Initialize(VideoTrack.Dimensions, FTimespan::FromSeconds(1.0 / VideoTrack.FrameRate)))
+		{
+			return;
+		}
+
+		// populate & add sample (on render thread)
+		const auto Settings = GetDefault<UAndroidMediaSettings>();
+
+		struct FWriteVideoSampleParams
+		{
+			TWeakPtr<FJavaAndroidMediaPlayer, ESPMode::ThreadSafe> JavaMediaPlayerPtr;
+			TWeakPtr<FMediaSamples, ESPMode::ThreadSafe> SamplesPtr;
+			TSharedRef<FAndroidMediaTextureSample, ESPMode::ThreadSafe> VideoSample;
+			int32 SampleCount;
+		}
+		WriteVideoSampleParams = { JavaMediaPlayer, Samples, VideoSample, (int32)(VideoTrack.Dimensions.X * VideoTrack.Dimensions.Y * sizeof(int32)) };
+
+		ENQUEUE_UNIQUE_RENDER_COMMAND_ONEPARAMETER(AndroidMediaPlayerWriteVideoSample,
+			FWriteVideoSampleParams, Params, WriteVideoSampleParams,
+			{
+				auto PinnedJavaMediaPlayer = Params.JavaMediaPlayerPtr.Pin();
+				auto PinnedSamples = Params.SamplesPtr.Pin();
+
+				if (!PinnedJavaMediaPlayer.IsValid() || !PinnedSamples.IsValid())
+				{
+					return;
+				}
+
+				const int32 CurrentFramePosition = PinnedJavaMediaPlayer->GetCurrentPosition();
+				const FTimespan Time = FTimespan::FromMilliseconds(CurrentFramePosition);
+
+				// write frame into buffer
+				void* Buffer = nullptr;
+				int64 SampleCount = 0;
+
+				if (!PinnedJavaMediaPlayer->GetVideoLastFrameData(Buffer, SampleCount))
+				{
+					return;
+				}
+
+				if (SampleCount != Params.SampleCount)
+				{
+					FPlatformMisc::LowLevelOutputDebugStringf(TEXT("FAndroidMediaPlayer::Fetch: Sample count mismatch (Buffer=%d, Available=%d"), Params.SampleCount, SampleCount);
+				}
+				check(Params.SampleCount <= SampleCount);
+
+				// must make a copy (buffer is owned by Java, not us!)
+				Params.VideoSample->InitializeBuffer(Buffer, Time, true);
+
+				PinnedSamples->AddVideo(Params.VideoSample);
+			});
+	}
+	else if (GSupportsImageExternal)
+	{
+		struct FWriteVideoSampleParams
+		{
+			TWeakPtr<FJavaAndroidMediaPlayer, ESPMode::ThreadSafe> JavaMediaPlayerPtr;
+			FGuid PlayerGuid;
+		};
+
+		FWriteVideoSampleParams WriteVideoSampleParams = { JavaMediaPlayer, PlayerGuid };
+
+		ENQUEUE_UNIQUE_RENDER_COMMAND_ONEPARAMETER(AndroidMediaPlayerWriteVideoSample,
+			FWriteVideoSampleParams, Params, WriteVideoSampleParams,
+			{
+				auto PinnedJavaMediaPlayer = Params.JavaMediaPlayerPtr.Pin();
+
+				if (!PinnedJavaMediaPlayer.IsValid())
+				{
+					return;
+				}
+
+				FTextureRHIRef VideoTexture = PinnedJavaMediaPlayer->GetVideoTexture();
+				if (VideoTexture == nullptr)
+				{
+					FRHIResourceCreateInfo CreateInfo;
+					VideoTexture = RHICmdList.CreateTextureExternal2D(1, 1, PF_R8G8B8A8, 1, 1, 0, CreateInfo);
+					PinnedJavaMediaPlayer->SetVideoTexture(VideoTexture);
+
+					if (VideoTexture == nullptr)
+					{
+						UE_LOG(LogAndroidMedia, Warning, TEXT("CreateTextureExternal2D failed!"));
+						return;
+					}
+
+					PinnedJavaMediaPlayer->SetVideoTextureValid(false);
+
+					#if ANDROIDMEDIAPLAYER_USE_NATIVELOGGING
+						FPlatformMisc::LowLevelOutputDebugStringf(TEXT("Fetch RT: Created VideoTexture: %d - %s"), *reinterpret_cast<int32*>(VideoTexture->GetNativeResource()), *Params.PlayerGuid.ToString());
+					#endif
+				}
+
+				int32 TextureId = *reinterpret_cast<int32*>(VideoTexture->GetNativeResource());
+				int32 CurrentFramePosition = 0;
+				bool bRegionChanged = false;
+				if (PinnedJavaMediaPlayer->UpdateVideoFrame(TextureId, &CurrentFramePosition, &bRegionChanged))
+				{
+					#if ANDROIDMEDIAPLAYER_USE_NATIVELOGGING
+//						FPlatformMisc::LowLevelOutputDebugStringf(TEXT("Fetch RT: %d - %s"), CurrentFramePosition, *Params.PlayerGuid.ToString());
+					#endif
+
+					// if region changed, need to reregister UV scale/offset
+					if (bRegionChanged)
+					{
+						PinnedJavaMediaPlayer->SetVideoTextureValid(false);
+
+						#if ANDROIDMEDIAPLAYER_USE_NATIVELOGGING
+							FPlatformMisc::LowLevelOutputDebugStringf(TEXT("Fetch RT: New UV Scale/Offset = {%f, %f}, {%f, %f} - %s"),
+								PinnedJavaMediaPlayer->GetUScale(), PinnedJavaMediaPlayer->GetUOffset(),
+								PinnedJavaMediaPlayer->GetVScale(), PinnedJavaMediaPlayer->GetVOffset(), *Params.PlayerGuid.ToString());
+						#endif
+					}
+				}
+
+				if (!PinnedJavaMediaPlayer->IsVideoTextureValid())
+				{
+					#if ANDROIDMEDIAPLAYER_USE_NATIVELOGGING
+						FPlatformMisc::LowLevelOutputDebugStringf(TEXT("Fetch RT: Register Guid: %s"), *Params.PlayerGuid.ToString());
+					#endif
+
+					FSamplerStateInitializerRHI SamplerStateInitializer(SF_Bilinear, AM_Clamp, AM_Clamp, AM_Clamp);
+					FSamplerStateRHIRef SamplerStateRHI = RHICreateSamplerState(SamplerStateInitializer);
+					FExternalTextureRegistry::Get().RegisterExternalTexture(Params.PlayerGuid, VideoTexture, SamplerStateRHI, FLinearColor(PinnedJavaMediaPlayer->GetUScale(), 0.0f, 0.0f, PinnedJavaMediaPlayer->GetVScale()), FLinearColor(PinnedJavaMediaPlayer->GetUOffset(), PinnedJavaMediaPlayer->GetVOffset(), 0.0f, 0.0f));
+
+					PinnedJavaMediaPlayer->SetVideoTextureValid(true);
+				}
+			});
+	}
+	else
+	{
+		// create new video sample
+		const FJavaAndroidMediaPlayer::FVideoTrack VideoTrack = VideoTracks[SelectedVideoTrack];
+		auto VideoSample = VideoSamplePool->AcquireShared();
+
+		if (!VideoSample->Initialize(VideoTrack.Dimensions, FTimespan::FromSeconds(1.0 / VideoTrack.FrameRate)))
+		{
+			return;
+		}
+
+		// populate & add sample (on render thread)
+		const auto Settings = GetDefault<UAndroidMediaSettings>();
+
+		struct FWriteVideoSampleParams
+		{
+			TWeakPtr<FJavaAndroidMediaPlayer, ESPMode::ThreadSafe> JavaMediaPlayerPtr;
+			TWeakPtr<FMediaSamples, ESPMode::ThreadSafe> SamplesPtr;
+			TSharedRef<FAndroidMediaTextureSample, ESPMode::ThreadSafe> VideoSample;
+			bool Cacheable;
+		}
+		WriteVideoSampleParams = { JavaMediaPlayer, Samples, VideoSample, Settings->CacheableVideoSampleBuffers };
+
+		ENQUEUE_UNIQUE_RENDER_COMMAND_ONEPARAMETER(AndroidMediaPlayerWriteVideoSample,
+			FWriteVideoSampleParams, Params, WriteVideoSampleParams,
+			{
+				auto PinnedJavaMediaPlayer = Params.JavaMediaPlayerPtr.Pin();
+				auto PinnedSamples = Params.SamplesPtr.Pin();
+
+				if (!PinnedJavaMediaPlayer.IsValid() || !PinnedSamples.IsValid())
+				{
+					return;
+				}
+
+				const int32 CurrentFramePosition = PinnedJavaMediaPlayer->GetCurrentPosition();
+				const FTimespan Time = FTimespan::FromMilliseconds(CurrentFramePosition);
+
+				// write frame into texture
+				FRHITexture2D* Texture = Params.VideoSample->InitializeTexture(Time);
+
+				if (Texture != nullptr)
+				{
+					int32 Resource = *reinterpret_cast<int32*>(Texture->GetNativeResource());
+					if (!PinnedJavaMediaPlayer->GetVideoLastFrame(Resource))
+					{
+						return;
+					}
+				}
+
+				PinnedSamples->AddVideo(Params.VideoSample);
+			});
+	}
+
+#else
+
+	// create new video sample
+	const FJavaAndroidMediaPlayer::FVideoTrack VideoTrack = VideoTracks[SelectedVideoTrack];
+	auto VideoSample = VideoSamplePool->AcquireShared();
+
+	if (!VideoSample->Initialize(VideoTrack.Dimensions, FTimespan::FromSeconds(1.0 / VideoTrack.FrameRate)))
+	{
+		return;
+	}
+
+	// populate & add sample (on render thread)
+	const auto Settings = GetDefault<UAndroidMediaSettings>();
+
+	struct FWriteVideoSampleParams
+	{
+		TWeakPtr<FJavaAndroidMediaPlayer, ESPMode::ThreadSafe> JavaMediaPlayerPtr;
+		TWeakPtr<FMediaSampleQueue, ESPMode::ThreadSafe> SamplesPtr;
+		TSharedRef<FAndroidMediaTextureSample, ESPMode::ThreadSafe> VideoSample;
+		int32 SampleCount;
+	};
+
+	FWriteVideoSampleParams WriteVideoSampleParams = { JavaMediaPlayer, Samples, VideoSample, (int32)(VideoTrack.Dimensions.X * VideoTrack.Dimensions.Y * sizeof(int32)) };
+
+	ENQUEUE_UNIQUE_RENDER_COMMAND_ONEPARAMETER(AndroidMediaPlayerWriteVideoSample,
+		FWriteVideoSampleParams, Params, WriteVideoSampleParams,
+		{
+			auto PinnedJavaMediaPlayer = Params.JavaMediaPlayerPtr.Pin();
+			auto PinnedSamples = Params.SamplesPtr.Pin();
+
+			if (!PinnedJavaMediaPlayer.IsValid() || !PinnedSamples.IsValid())
+			{
+				return;
+			}
+
+			const int32 CurrentFramePosition = PinnedJavaMediaPlayer->GetCurrentPosition();
+			const FTimespan Time = FTimespan::FromMilliseconds(CurrentFramePosition);
+
+			// write frame into buffer
+			void* Buffer = nullptr;
+			int64 SampleCount = 0;
+
+			if (!PinnedJavaMediaPlayer->GetVideoLastFrameData(Buffer, SampleCount))
+			{
+				return;
+			}
+
+			if (SampleCount != Params.SampleCount)
+			{
+				FPlatformMisc::LowLevelOutputDebugStringf(TEXT("FAndroidMediaPlayer::Fetch: Sample count mismatch (Buffer=%d, Available=%d"), Params.SampleCount, SampleCount);
+			}
+			check(Params.SampleCount <= SampleCount);
+
+			// must make a copy (buffer is owned by Java, not us!)
+			Params.VideoSample->InitializeBuffer(Buffer, Time, true);
+
+			PinnedSamples->AddVideo(Params.VideoSample);
+		});
+#endif //WITH_ENGINE
 }
 
 
-void FAndroidMediaPlayer::TickPlayer(float DeltaTime)
+void FAndroidMediaPlayer::TickInput(FTimespan DeltaTime)
 {
-	if (State != EMediaState::Playing)
+	if (CurrentState != EMediaState::Playing)
 	{
 		// remove delegates if registered
 		if (ResumeHandle.IsValid())
@@ -374,38 +596,639 @@ void FAndroidMediaPlayer::TickPlayer(float DeltaTime)
 			PauseHandle.Reset();
 		}
 
+#if ANDROIDMEDIAPLAYER_USE_PREPAREASYNC
+		// if preparing, see if finished
+		if (CurrentState == EMediaState::Preparing)
+		{
+			if (!JavaMediaPlayer.IsValid())
+			{
+				return;
+			}
+
+			if (JavaMediaPlayer->IsPrepared())
+			{
+				InitializePlayer();
+			}
+		}
+#endif
+
+		return;
+	}
+
+	if (!JavaMediaPlayer.IsValid())
+	{
 		return;
 	}
 
 	// register delegate if not registered
 	if (!ResumeHandle.IsValid())
 	{
-		ResumeHandle = FCoreDelegates::ApplicationHasEnteredForegroundDelegate.AddRaw(this, &FAndroidMediaPlayer::AppServiceResume);
+		ResumeHandle = FCoreDelegates::ApplicationHasEnteredForegroundDelegate.AddRaw(this, &FAndroidMediaPlayer::HandleApplicationWillEnterBackground);
 	}
 	if (!PauseHandle.IsValid())
 	{
-		PauseHandle = FCoreDelegates::ApplicationWillEnterBackgroundDelegate.AddRaw(this, &FAndroidMediaPlayer::AppServicePause);
+		PauseHandle = FCoreDelegates::ApplicationWillEnterBackgroundDelegate.AddRaw(this, &FAndroidMediaPlayer::HandleApplicationHasEnteredForeground);
 	}
 
-	Tracks.Tick();
-
+	// generate events
 	if (!JavaMediaPlayer->IsPlaying())
 	{
-		//FPlatformMisc::LowLevelOutputDebugString(TEXT("STOPPED!!!!!"));
-		State = EMediaState::Stopped;
-		MediaEvent.Broadcast(EMediaEvent::PlaybackEndReached);
-	}
-	else if (Tracks.DidPlaybackLoop())
-	{
-		//FPlatformMisc::LowLevelOutputDebugString(TEXT("LOOPED!!!!!"));
-		MediaEvent.Broadcast(EMediaEvent::PlaybackEndReached);
-	}
-}
+		if (JavaMediaPlayer->DidComplete())
+		{
+			EventSink.ReceiveMediaEvent(EMediaEvent::PlaybackEndReached);
 
+			#if ANDROIDMEDIAPLAYER_USE_NATIVELOGGING
+				FPlatformMisc::LowLevelOutputDebugStringf(TEXT("FAndroidMedia::Tick - PlaybackEndReached - !playing - %s"), *PlayerGuid.ToString());
+			#endif
+		}
 
-void FAndroidMediaPlayer::TickVideo(float DeltaTime)
-{
+		// might catch it restarting the loop so ignore if looping
+		if (!bLooping)
+		{
+			CurrentState = EMediaState::Stopped;
+			EventSink.ReceiveMediaEvent(EMediaEvent::PlaybackSuspended);
+
+			#if ANDROIDMEDIAPLAYER_USE_NATIVELOGGING
+				FPlatformMisc::LowLevelOutputDebugStringf(TEXT("FAndroidMedia::Tick - PlaybackSuspended - !playing - %s"), *PlayerGuid.ToString());
+			#endif
+		}
+	}
+	else if (JavaMediaPlayer->DidComplete())
+	{
+		EventSink.ReceiveMediaEvent(EMediaEvent::PlaybackEndReached);
+
+		#if ANDROIDMEDIAPLAYER_USE_NATIVELOGGING
+			FPlatformMisc::LowLevelOutputDebugStringf(TEXT("FAndroidMedia::Tick - PlaybackEndReached - DidComplete true - %s"), *PlayerGuid.ToString());
+		#endif
+	}
+}
+
+
+/* FAndroidMediaPlayer implementation
+ *****************************************************************************/
+
+bool FAndroidMediaPlayer::InitializePlayer()
+{
+	#if ANDROIDMEDIAPLAYER_USE_NATIVELOGGING
+		FPlatformMisc::LowLevelOutputDebugStringf(TEXT("FAndroidMedia::InitializePlayer %s"), *PlayerGuid.ToString());
+	#endif
+
+	if (ResumeHandle.IsValid())
+	{
+		FCoreDelegates::ApplicationHasEnteredForegroundDelegate.Remove(ResumeHandle);
+		ResumeHandle.Reset();
+	}
+
+	if (PauseHandle.IsValid())
+	{
+		FCoreDelegates::ApplicationWillEnterBackgroundDelegate.Remove(PauseHandle);
+		PauseHandle.Reset();
+	}
+
+	JavaMediaPlayer->GetAudioTracks(AudioTracks);
+	JavaMediaPlayer->GetCaptionTracks(CaptionTracks);
+	JavaMediaPlayer->GetVideoTracks(VideoTracks);
+
+	for (FJavaAndroidMediaPlayer::FVideoTrack Track : VideoTracks)
+	{
+		Info += FString::Printf(TEXT("Stream %i\n"), Track.Index);
+		Info += TEXT("    Type: Video\n");
+//		Info += FString::Printf(TEXT("    Duration: %s\n"), *FTimespan::FromMilliseconds(Track->StreamInfo.duration).ToString());
+		Info += FString::Printf(TEXT("    MimeType: %s\n"), *Track.MimeType);
+		Info += FString::Printf(TEXT("    Language: %s\n"), *Track.Language);
+		Info += FString::Printf(TEXT("    FrameRate: %0.1f\n"), Track.FrameRate);
+		Info += FString::Printf(TEXT("    Dimensions: %i x %i\n"), Track.Dimensions.X, Track.Dimensions.Y);
+		Info += TEXT("\n");
+	}
+
+	for (FJavaAndroidMediaPlayer::FAudioTrack Track : AudioTracks)
+	{
+		Info += FString::Printf(TEXT("Stream %i\n"), Track.Index);
+		Info += TEXT("    Type: Audio\n");
+//		Info += FString::Printf(TEXT("    Duration: %s\n"), *FTimespan::FromMilliseconds(Track->StreamInfo.duration).ToString());
+		Info += FString::Printf(TEXT("    MimeType: %s\n"), *Track.MimeType);
+		Info += FString::Printf(TEXT("    Language: %s\n"), *Track.Language);
+		Info += FString::Printf(TEXT("    Channels: %i\n"), Track.Channels);
+		Info += FString::Printf(TEXT("    Sample Rate: %i Hz\n"), Track.SampleRate);
+		Info += TEXT("\n");
+	}
+
+	for (FJavaAndroidMediaPlayer::FCaptionTrack Track : CaptionTracks)
+	{
+		Info += FString::Printf(TEXT("Stream %i\n"), Track.Index);
+		Info += TEXT("    Type: Caption\n");
+		Info += FString::Printf(TEXT("    MimeType: %s\n"), *Track.MimeType);
+		Info += FString::Printf(TEXT("    Language: %s\n"), *Track.Language);
+		Info += TEXT("\n");
+	}
+
+	// Select first audio and video track by default (have to force update first time by setting to none)
+	if (AudioTracks.Num() > 0)
+	{
+		JavaMediaPlayer->SetAudioEnabled(true);
+		SelectedAudioTrack = 0;
+	}
+	else
+	{
+		JavaMediaPlayer->SetAudioEnabled(false);
+		SelectedAudioTrack = INDEX_NONE;
+	}
+
+	if (VideoTracks.Num() > 0)
+	{
+		JavaMediaPlayer->SetVideoEnabled(true);
+		SelectedVideoTrack = 0;
+	}
+	else
+	{
+		JavaMediaPlayer->SetVideoEnabled(false);
+		SelectedVideoTrack = INDEX_NONE;
+	}
+
+	CurrentState = EMediaState::Stopped;
+
+	// notify listeners
+	EventSink.ReceiveMediaEvent(EMediaEvent::TracksChanged);
+	EventSink.ReceiveMediaEvent(EMediaEvent::MediaOpened);
+
+	return true;
+}
+
+
+/* IMediaTracks interface
+ *****************************************************************************/
+
+bool FAndroidMediaPlayer::GetAudioTrackFormat(int32 TrackIndex, int32 FormatIndex, FMediaAudioTrackFormat& OutFormat) const
+{
+	if ((FormatIndex != 0) || !AudioTracks.IsValidIndex(TrackIndex))
+	{
+		return false;
+	}
+
+	const FJavaAndroidMediaPlayer::FAudioTrack& Track = AudioTracks[TrackIndex];
+
+	OutFormat.BitsPerSample = 16;
+	OutFormat.NumChannels = Track.Channels;
+	OutFormat.SampleRate = Track.SampleRate;
+	OutFormat.TypeName = TEXT("Native");
+
+	return true;
+}
+
+
+int32 FAndroidMediaPlayer::GetNumTracks(EMediaTrackType TrackType) const
+{
+	switch (TrackType)
+	{
+	case EMediaTrackType::Audio:
+		return AudioTracks.Num();
+
+	case EMediaTrackType::Caption:
+		return CaptionTracks.Num();
+
+	case EMediaTrackType::Video:
+		return VideoTracks.Num();
+
+	default:
+		return 0;
+	}
+}
+
+
+int32 FAndroidMediaPlayer::GetNumTrackFormats(EMediaTrackType TrackType, int32 TrackIndex) const
+{
+	// @todo chrisb: add support for track formats
+	return ((TrackIndex >= 0) && (TrackIndex < GetNumTracks(TrackType))) ? 1 : 0;
+}
+
+
+int32 FAndroidMediaPlayer::GetSelectedTrack(EMediaTrackType TrackType) const
+{
+	switch (TrackType)
+	{
+	case EMediaTrackType::Audio:
+		return SelectedAudioTrack;
+
+	case EMediaTrackType::Caption:
+		return SelectedCaptionTrack;
+
+	case EMediaTrackType::Video:
+		return SelectedVideoTrack;
+
+	default:
+		return INDEX_NONE;
+	}
+}
 
+
+FText FAndroidMediaPlayer::GetTrackDisplayName(EMediaTrackType TrackType, int32 TrackIndex) const
+{
+	switch (TrackType)
+	{
+	case EMediaTrackType::Audio:
+		if (AudioTracks.IsValidIndex(TrackIndex))
+		{
+			return FText::FromString(AudioTracks[TrackIndex].DisplayName);
+		}
+		break;
+
+	case EMediaTrackType::Caption:
+		if (CaptionTracks.IsValidIndex(TrackIndex))
+		{
+			return FText::FromString(CaptionTracks[TrackIndex].DisplayName);
+		}
+		break;
+
+	case EMediaTrackType::Video:
+		if (VideoTracks.IsValidIndex(TrackIndex))
+		{
+			return FText::FromString(VideoTracks[TrackIndex].DisplayName);
+		}
+
+	default:
+		break;
+	}
+
+	return FText::GetEmpty();
+}
+
+
+int32 FAndroidMediaPlayer::GetTrackFormat(EMediaTrackType TrackType, int32 TrackIndex) const
+{
+	// @todo chrisb: add support for track formats
+	return (GetSelectedTrack(TrackType) != INDEX_NONE) ? 0 : INDEX_NONE;
+}
+
+
+FString FAndroidMediaPlayer::GetTrackLanguage(EMediaTrackType TrackType, int32 TrackIndex) const
+{
+	switch (TrackType)
+	{
+	case EMediaTrackType::Audio:
+		if (AudioTracks.IsValidIndex(TrackIndex))
+		{
+			return AudioTracks[TrackIndex].Language;
+		}
+		break;
+
+	case EMediaTrackType::Caption:
+		if (CaptionTracks.IsValidIndex(TrackIndex))
+		{
+			return CaptionTracks[TrackIndex].Language;
+		}
+		break;
+
+	case EMediaTrackType::Video:
+		if (VideoTracks.IsValidIndex(TrackIndex))
+		{
+			return VideoTracks[TrackIndex].Language;
+		}
+
+	default:
+		break;
+	}
+
+	return FString();
+}
+
+
+FString FAndroidMediaPlayer::GetTrackName(EMediaTrackType TrackType, int32 TrackIndex) const
+{
+	switch (TrackType)
+	{
+	case EMediaTrackType::Audio:
+		if (AudioTracks.IsValidIndex(TrackIndex))
+		{
+			return AudioTracks[TrackIndex].Name;
+		}
+		break;
+
+	case EMediaTrackType::Caption:
+		if (CaptionTracks.IsValidIndex(TrackIndex))
+		{
+			return CaptionTracks[TrackIndex].Name;
+		}
+		break;
+
+	case EMediaTrackType::Video:
+		if (VideoTracks.IsValidIndex(TrackIndex))
+		{
+			return VideoTracks[TrackIndex].Name;
+		}
+
+	default:
+		break;
+	}
+
+	return FString();
+}
+
+
+bool FAndroidMediaPlayer::GetVideoTrackFormat(int32 TrackIndex, int32 FormatIndex, FMediaVideoTrackFormat& OutFormat) const
+{
+	if ((FormatIndex != 0) || !VideoTracks.IsValidIndex(TrackIndex))
+	{
+		return false;
+	}
+
+	const FJavaAndroidMediaPlayer::FVideoTrack& Track = VideoTracks[TrackIndex];
+
+	OutFormat.Dim = Track.Dimensions;
+	OutFormat.FrameRate = Track.FrameRate;
+	OutFormat.FrameRates = TRange<float>(Track.FrameRate);
+	OutFormat.TypeName = TEXT("BGRA");
+
+	return true;
+}
+
+
+bool FAndroidMediaPlayer::SelectTrack(EMediaTrackType TrackType, int32 TrackIndex)
+{
+	switch (TrackType)
+	{
+	case EMediaTrackType::Audio:
+		if (TrackIndex != SelectedAudioTrack)
+		{
+			UE_LOG(LogAndroidMedia, Verbose, TEXT("Player %llx: Selecting audio track %i instead of %i (%i tracks)"), this, TrackIndex, SelectedVideoTrack, AudioTracks.Num());
+
+			if (TrackIndex == INDEX_NONE)
+			{
+				UE_LOG(LogAndroidMedia, VeryVerbose, TEXT("Player %llx: Disabling audio"), this, TrackIndex);
+
+				JavaMediaPlayer->SetAudioEnabled(false);
+			}
+			else
+			{
+				if (!AudioTracks.IsValidIndex(TrackIndex) || !JavaMediaPlayer->SelectTrack(AudioTracks[TrackIndex].Index))
+				{
+					return false;
+				}
+
+				UE_LOG(LogAndroidMedia, VeryVerbose, TEXT("Player %llx: Enabling audio"), this, TrackIndex);
+
+				JavaMediaPlayer->SetAudioEnabled(true);
+			}
+
+			SelectedAudioTrack = TrackIndex;
+		}
+		break;
+
+	case EMediaTrackType::Caption:
+		if (TrackIndex != SelectedCaptionTrack)
+		{
+			UE_LOG(LogAndroidMedia, Verbose, TEXT("Player %llx: Selecting caption track %i instead of %i (%i tracks)"), this, TrackIndex, SelectedCaptionTrack, CaptionTracks.Num());
+
+			if (TrackIndex == INDEX_NONE)
+			{
+				UE_LOG(LogAndroidMedia, VeryVerbose, TEXT("Player %llx: Disabling captions"), this, TrackIndex);
+			}
+			else
+			{
+				if (!CaptionTracks.IsValidIndex(TrackIndex) || JavaMediaPlayer->SelectTrack(CaptionTracks[TrackIndex].Index))
+				{
+					return false;
+				}
+
+				UE_LOG(LogAndroidMedia, VeryVerbose, TEXT("Player %llx: Enabling captions"), this, TrackIndex);
+			}
+
+			SelectedCaptionTrack = TrackIndex;
+		}
+		break;
+
+	case EMediaTrackType::Video:
+		if (TrackIndex != SelectedVideoTrack)
+		{
+			UE_LOG(LogAndroidMedia, Verbose, TEXT("Player %llx: Selecting video track %i instead of %i (%i tracks)."), this, TrackIndex, SelectedVideoTrack, VideoTracks.Num());
+
+			if (TrackIndex == INDEX_NONE)
+			{
+				UE_LOG(LogAndroidMedia, VeryVerbose, TEXT("Player %llx: Disabling video"), this, TrackIndex);
+				JavaMediaPlayer->SetVideoEnabled(false);
+			}
+			else
+			{
+				if (!VideoTracks.IsValidIndex(TrackIndex) || !JavaMediaPlayer->SelectTrack(VideoTracks[TrackIndex].Index))
+				{
+					return false;
+				}
+
+				UE_LOG(LogAndroidMedia, VeryVerbose, TEXT("Player %llx: Enabling video"), this, TrackIndex);
+				JavaMediaPlayer->SetVideoEnabled(true);
+			}
+
+			SelectedVideoTrack = TrackIndex;
+		}
+		break;
+
+	default:
+		return false;
+	}
+
+	return true;
+}
+
+
+bool FAndroidMediaPlayer::SetTrackFormat(EMediaTrackType TrackType, int32 TrackIndex, int32 FormatIndex)
+{
+	if (FormatIndex != 0)
+	{
+		return false;
+	}
+
+	switch (TrackType)
+	{
+	case EMediaTrackType::Audio:
+		return AudioTracks.IsValidIndex(TrackIndex);
+
+	case EMediaTrackType::Caption:
+		return CaptionTracks.IsValidIndex(TrackIndex);
+
+	case EMediaTrackType::Video:
+		return VideoTracks.IsValidIndex(TrackIndex);
+
+	default:
+		return false;
+	}
+}
+
+
+/* IMediaControls interface
+ *****************************************************************************/
+
+bool FAndroidMediaPlayer::CanControl(EMediaControl Control) const
+{
+	if (Control == EMediaControl::Pause)
+	{
+		return (CurrentState == EMediaState::Playing);
+	}
+
+	if (Control == EMediaControl::Resume)
+	{
+		return ((CurrentState == EMediaState::Playing) || (CurrentState == EMediaState::Stopped));
+	}
+
+	if (Control == EMediaControl::Seek)
+	{
+		return ((CurrentState != EMediaState::Closed) && (CurrentState != EMediaState::Error));
+	}
+
+	return false;
+}
+
+
+FTimespan FAndroidMediaPlayer::GetDuration() const
+{
+	if (CurrentState == EMediaState::Error)
+	{
+		return FTimespan::Zero();
+	}
+
+	return FTimespan::FromMilliseconds(JavaMediaPlayer->GetDuration());
+}
+
+
+float FAndroidMediaPlayer::GetRate() const
+{
+	return (CurrentState == EMediaState::Playing) ? 1.0f : 0.0f;
+}
+
+
+EMediaState FAndroidMediaPlayer::GetState() const
+{
+	return CurrentState;
+}
+
+
+EMediaStatus FAndroidMediaPlayer::GetStatus() const
+{
+	return EMediaStatus::None;
+}
+
+
+TRangeSet<float> FAndroidMediaPlayer::GetSupportedRates(EMediaRateThinning /*Thinning*/) const
+{
+	TRangeSet<float> Result;
+
+	Result.Add(TRange<float>(0.0f));
+	Result.Add(TRange<float>(1.0f));
+
+	return Result;
+}
+
+
+FTimespan FAndroidMediaPlayer::GetTime() const
+{
+	if ((CurrentState == EMediaState::Closed) || (CurrentState == EMediaState::Error))
+	{
+		return FTimespan::Zero();
+	}
+
+	return FTimespan::FromMilliseconds(JavaMediaPlayer->GetCurrentPosition());
+}
+
+
+bool FAndroidMediaPlayer::IsLooping() const
+{
+	return bLooping;
+}
+
+
+bool FAndroidMediaPlayer::Seek(const FTimespan& Time)
+{
+	UE_LOG(LogAndroidMedia, Verbose, TEXT("Player %llx: Seeking to %s"), this, *Time.ToString());
+
+	if ((CurrentState == EMediaState::Closed) ||
+		(CurrentState == EMediaState::Error) ||
+		(CurrentState == EMediaState::Preparing))
+	{
+		UE_LOG(LogAndroidMedia, Warning, TEXT("Cannot seek while closed, preparing, or in error state"));
+		return false;
+	}
+
+	UE_LOG(LogAndroidMedia, Verbose, TEXT("Player %llx: Seeking to %s"), this, *Time.ToString());
+
+	JavaMediaPlayer->SeekTo(static_cast<int32>(Time.GetTotalMilliseconds()));
+	EventSink.ReceiveMediaEvent(EMediaEvent::SeekCompleted);
+
+	return true;
+}
+
+
+bool FAndroidMediaPlayer::SetLooping(bool Looping)
+{
+	bLooping = Looping;
+
+	if (!JavaMediaPlayer.IsValid())
+	{
+		return false;
+	}
+
+	JavaMediaPlayer->SetLooping(Looping);
+
+	return true;
+}
+
+
+bool FAndroidMediaPlayer::SetRate(float Rate)
+{
+	if ((CurrentState == EMediaState::Closed) ||
+		(CurrentState == EMediaState::Error) ||
+		(CurrentState == EMediaState::Preparing))
+	{
+		UE_LOG(LogAndroidMedia, Warning, TEXT("Cannot set rate while closed, preparing, or in error state"));
+		return false;
+	}
+
+	if (Rate == GetRate())
+	{
+		return true; // rate already set
+	}
+
+	UE_LOG(LogAndroidMedia, Verbose, TEXT("Player %llx: Setting rate from to %f to %f"), this, GetRate(), Rate);
+
+	if (Rate == 0.0f)
+	{
+		JavaMediaPlayer->Pause();
+		CurrentState = EMediaState::Paused;
+		EventSink.ReceiveMediaEvent(EMediaEvent::PlaybackSuspended);
+	}
+	else if (Rate == 1.0f)
+	{
+		JavaMediaPlayer->Start();
+		CurrentState = EMediaState::Playing;
+		EventSink.ReceiveMediaEvent(EMediaEvent::PlaybackResumed);
+	}
+	else
+	{
+		UE_LOG(LogAndroidMedia, Warning, TEXT("The rate %f is not supported"), Rate);
+		return false;
+	}
+
+	return true;
+}
+
+
+/* FAndroidMediaPlayer callbacks
+ *****************************************************************************/
+
+void FAndroidMediaPlayer::HandleApplicationHasEnteredForeground()
+{
+	// check state in case changed before ticked
+	if ((CurrentState == EMediaState::Playing) && JavaMediaPlayer.IsValid())
+	{
+		JavaMediaPlayer->Pause();
+	}
+}
+
+
+void FAndroidMediaPlayer::HandleApplicationWillEnterBackground()
+{
+	// check state in case changed before ticked
+	if ((CurrentState == EMediaState::Playing) && JavaMediaPlayer.IsValid())
+	{
+		JavaMediaPlayer->Start();
+	}
 }
 
 
diff --git a/Engine/Plugins/Media/AndroidMedia/Source/AndroidMedia/Private/Player/AndroidMediaPlayer.h b/Engine/Plugins/Media/AndroidMedia/Source/AndroidMedia/Private/Player/AndroidMediaPlayer.h
--- a/Engine/Plugins/Media/AndroidMedia/Source/AndroidMedia/Private/Player/AndroidMediaPlayer.h
+++ b/Engine/Plugins/Media/AndroidMedia/Source/AndroidMedia/Private/Player/AndroidMediaPlayer.h
@@ -2,94 +2,168 @@
 
 #pragma once
 
-#include "CoreMinimal.h"
-#include "AndroidMediaTracks.h"
+#include "CoreTypes.h"
+#include "Containers/Array.h"
+#include "Containers/UnrealString.h"
+#include "Misc/Guid.h"
+#include "IMediaCache.h"
 #include "IMediaControls.h"
 #include "IMediaPlayer.h"
-#include "Containers/Ticker.h"
+#include "IMediaTracks.h"
+#include "IMediaView.h"
+#include "Misc/Timespan.h"
+#include "Templates/SharedPointer.h"
 
+#include "Android/AndroidJavaMediaPlayer.h"
 
-class FJavaAndroidMediaPlayer;
+class FAndroidMediaTextureSamplePool;
+class FMediaSamples;
+class IMediaEventSink;
 
 
 /*
-	Implement media playback using the Android MediaPlayer interface.
-*/
+ * Implements media playback using the Android MediaPlayer interface.
+ */
 class FAndroidMediaPlayer
-	: public IMediaControls
-	, public IMediaPlayer
+	: public IMediaPlayer
+	, protected IMediaCache
+	, protected IMediaControls
+	, protected IMediaTracks
+	, protected IMediaView
 {
 public:
 	
-	/** Default constructor. */
-	FAndroidMediaPlayer();
+	/**
+	 * Create and initialize a new instance.
+	 *
+	 * @param InEventSink The object that receives media events from this player.
+	 */
+	FAndroidMediaPlayer(IMediaEventSink& InEventSink);
 
 	/** Virtual destructor. */
 	virtual ~FAndroidMediaPlayer();
 
 public:
 
+	//~ IMediaPlayer interface
+
+	virtual void Close() override;
+	virtual IMediaCache& GetCache() override;
+	virtual IMediaControls& GetControls() override;
+	virtual FString GetInfo() const override;
+	virtual FName GetName() const override;
+	virtual IMediaSamples& GetSamples() override;
+	virtual FString GetStats() const override;
+	virtual IMediaTracks& GetTracks() override;
+	virtual FString GetUrl() const override;
+	virtual IMediaView& GetView() override;
+	virtual bool Open(const FString& Url, const IMediaOptions* Options) override;
+	virtual bool Open(const TSharedRef<FArchive, ESPMode::ThreadSafe>& Archive, const FString& OriginalUrl, const IMediaOptions* Options) override;
+	virtual void SetGuid(const FGuid& Guid) override;
+	virtual void TickFetch(FTimespan DeltaTime) override;
+	virtual void TickInput(FTimespan DeltaTime) override;
+
+protected:
+
+	/**
+	 * Initialize the media player.
+	 *
+	 * @return true on success, false otherwise.
+	 */
+	bool InitializePlayer();
+
+protected:
+
 	//~ IMediaControls interface
 
+	virtual bool CanControl(EMediaControl Control) const override;
 	virtual FTimespan GetDuration() const override;
 	virtual float GetRate() const override;
 	virtual EMediaState GetState() const override;
-	virtual TRange<float> GetSupportedRates(EMediaPlaybackDirections Direction, bool Unthinned) const override;
+	virtual EMediaStatus GetStatus() const override;
+	virtual TRangeSet<float> GetSupportedRates(EMediaRateThinning Thinning) const override;
 	virtual FTimespan GetTime() const override;
 	virtual bool IsLooping() const override;
 	virtual bool Seek(const FTimespan& Time) override;
 	virtual bool SetLooping(bool Looping) override;
 	virtual bool SetRate(float Rate) override;
-	virtual bool SupportsRate(float Rate, bool Unthinned) const override;
-	virtual bool SupportsScrubbing() const override;
-	virtual bool SupportsSeeking() const override;
 
-public:
+protected:
 
-	//~ IMediaPlayer interface
+	//~ IMediaTracks interface
 
-	virtual void Close() override;
-	virtual IMediaControls& GetControls() override;
-	virtual FString GetInfo() const override;
-	virtual FName GetName() const override;
-	virtual IMediaOutput& GetOutput() override;
-	virtual FString GetStats() const override;
-	virtual IMediaTracks& GetTracks() override;
-	virtual FString GetUrl() const override;
-	virtual bool Open(const FString& Url, const IMediaOptions& Options) override;
-	virtual bool Open(const TSharedRef<FArchive, ESPMode::ThreadSafe>& Archive, const FString& OriginalUrl, const IMediaOptions& Options) override;
-	virtual void TickPlayer(float DeltaTime) override;
-	virtual void TickVideo(float DeltaTime) override;
+	virtual bool GetAudioTrackFormat(int32 TrackIndex, int32 FormatIndex, FMediaAudioTrackFormat& OutFormat) const override;
+	virtual int32 GetNumTracks(EMediaTrackType TrackType) const override;
+	virtual int32 GetNumTrackFormats(EMediaTrackType TrackType, int32 TrackIndex) const override;
+	virtual int32 GetSelectedTrack(EMediaTrackType TrackType) const override;
+	virtual FText GetTrackDisplayName(EMediaTrackType TrackType, int32 TrackIndex) const override;
+	virtual int32 GetTrackFormat(EMediaTrackType TrackType, int32 TrackIndex) const override;
+	virtual FString GetTrackLanguage(EMediaTrackType TrackType, int32 TrackIndex) const override;
+	virtual FString GetTrackName(EMediaTrackType TrackType, int32 TrackIndex) const override;
+	virtual bool GetVideoTrackFormat(int32 TrackIndex, int32 FormatIndex, FMediaVideoTrackFormat& OutFormat) const override;
+	virtual bool SelectTrack(EMediaTrackType TrackType, int32 TrackIndex) override;
+	virtual bool SetTrackFormat(EMediaTrackType TrackType, int32 TrackIndex, int32 FormatIndex) override;
 
-	DECLARE_DERIVED_EVENT(FVlcMediaPlayer, IMediaPlayer::FOnMediaEvent, FOnMediaEvent);
-	virtual FOnMediaEvent& OnMediaEvent() override
-	{
-		return MediaEvent;
-	}
+private:
+
+	/** Callback for when the application resumed in the foreground. */
+	void HandleApplicationHasEnteredForeground();
+
+	/** Callback for when the application is being paused in the background. */
+	void HandleApplicationWillEnterBackground();
 
 private:
-	void AppServicePause();
-	void AppServiceResume();
 
-	/** The Java side media interface. */
-	TSharedPtr<FJavaAndroidMediaPlayer, ESPMode::ThreadSafe> JavaMediaPlayer;
+	/** Audio track descriptors. */
+	TArray<FJavaAndroidMediaPlayer::FAudioTrack> AudioTracks;
 
-	/** Holds an event delegate that is invoked when a media event occurred. */
-	FOnMediaEvent MediaEvent;
+	/** Caption track descriptors. */
+	TArray<FJavaAndroidMediaPlayer::FCaptionTrack> CaptionTracks;
 
-	/** Currently opened media. */
-	FString MediaUrl;
+	/** Video track descriptors. */
+	TArray<FJavaAndroidMediaPlayer::FVideoTrack> VideoTracks;
+
+private:
+
+	/** Current player state. */
+	EMediaState CurrentState;
+
+	/** Current state of looping. */
+	bool bLooping;
+
+	/** The media event handler. */
+	IMediaEventSink& EventSink;
 
 	/** Media information string. */
 	FString Info;
 
-	/** Our understanding of the state of the Java media player. */
-	EMediaState State;
+	/** The Java side media interface. */
+	TSharedPtr<FJavaAndroidMediaPlayer, ESPMode::ThreadSafe> JavaMediaPlayer;
+
+	/** Currently opened media. */
+	FString MediaUrl;
 
-	/** Track collection. */
-	FAndroidMediaTracks Tracks;
+	/** Media player Guid */
+	FGuid PlayerGuid;
 
-	/** Foreground/background delegates */
+	/** Foreground/background delegate for pause. */
 	FDelegateHandle PauseHandle;
+
+	/** Foreground/background delegate for resume. */
 	FDelegateHandle ResumeHandle;
+
+	/** The media sample queue. */
+	TSharedPtr<FMediaSamples, ESPMode::ThreadSafe> Samples;
+
+	/** Index of the selected audio track. */
+	int32 SelectedAudioTrack;
+
+	/** Index of the selected audio track. */
+	int32 SelectedCaptionTrack;
+
+	/** Index of the selected video track. */
+	int32 SelectedVideoTrack;
+
+	/** Video sample object pool. */
+	FAndroidMediaTextureSamplePool* VideoSamplePool;
 };
diff --git a/Engine/Plugins/Media/AndroidMedia/Source/AndroidMedia/Public/IAndroidMediaModule.h b/Engine/Plugins/Media/AndroidMedia/Source/AndroidMedia/Public/IAndroidMediaModule.h
--- a/Engine/Plugins/Media/AndroidMedia/Source/AndroidMedia/Public/IAndroidMediaModule.h
+++ b/Engine/Plugins/Media/AndroidMedia/Source/AndroidMedia/Public/IAndroidMediaModule.h
@@ -2,12 +2,13 @@
 
 #pragma once
 
-#include "CoreTypes.h"
 #include "Templates/SharedPointer.h"
 #include "Modules/ModuleInterface.h"
 
+class IMediaEventSink;
 class IMediaPlayer;
 
+
 /**
  * Interface for the AndroidMedia module.
  */
@@ -17,11 +18,12 @@ class IAndroidMediaModule
 public:
 
 	/**
-	 * Creates a Android based media player.
+	 * Create a Android based media player.
 	 *
+	 * @param EventSink The object that receives media events from the player.
 	 * @return A new media player, or nullptr if a player couldn't be created.
 	 */
-	virtual TSharedPtr<IMediaPlayer, ESPMode::ThreadSafe> CreatePlayer() = 0;
+	virtual TSharedPtr<IMediaPlayer, ESPMode::ThreadSafe> CreatePlayer(IMediaEventSink& EventSink) = 0;
 
 public:
 
diff --git a/Engine/Plugins/Media/AndroidMedia/Source/AndroidMediaFactory/AndroidMediaFactory.Build.cs b/Engine/Plugins/Media/AndroidMedia/Source/AndroidMediaFactory/AndroidMediaFactory.Build.cs
--- a/Engine/Plugins/Media/AndroidMedia/Source/AndroidMediaFactory/AndroidMediaFactory.Build.cs
+++ b/Engine/Plugins/Media/AndroidMedia/Source/AndroidMediaFactory/AndroidMediaFactory.Build.cs
@@ -6,32 +6,32 @@ namespace UnrealBuildTool.Rules
 	{
 		public AndroidMediaFactory(ReadOnlyTargetRules Target) : base(Target)
 		{
-            DynamicallyLoadedModuleNames.AddRange(
-                new string[] {
-                    "Media",
-				}
-            );
+			DynamicallyLoadedModuleNames.AddRange(
+				new string[] {
+					"Media",
+				});
 
 			PrivateDependencyModuleNames.AddRange(
 				new string[] {
-					"Core",
-                    "CoreUObject",
-                    "MediaAssets",
-                }
-			);
+					"MediaAssets",
+				});
 
 			PrivateIncludePathModuleNames.AddRange(
 				new string[] {
 					"AndroidMedia",
-                    "Media",
-				}
-			);
+					"Media",
+				});
 
 			PrivateIncludePaths.AddRange(
 				new string[] {
 					"AndroidMediaFactory/Private",
-				}
-			);
+				});
+
+			PublicDependencyModuleNames.AddRange(
+				new string[] {
+					"Core",
+					"CoreUObject",
+				});
 
 			if (Target.Platform == UnrealTargetPlatform.Android)
 			{
diff --git a/Engine/Plugins/Media/AndroidMedia/Source/AndroidMediaFactory/Private/AndroidMediaFactoryModule.cpp b/Engine/Plugins/Media/AndroidMedia/Source/AndroidMediaFactory/Private/AndroidMediaFactoryModule.cpp
--- a/Engine/Plugins/Media/AndroidMedia/Source/AndroidMediaFactory/Private/AndroidMediaFactoryModule.cpp
+++ b/Engine/Plugins/Media/AndroidMedia/Source/AndroidMediaFactory/Private/AndroidMediaFactoryModule.cpp
@@ -1,13 +1,19 @@
 // Copyright 1998-2017 Epic Games, Inc. All Rights Reserved.
 
-#include "CoreMinimal.h"
+#include "AndroidMediaFactoryPrivate.h"
+
+#include "CoreTypes.h"
+#include "Containers/Array.h"
+#include "Containers/UnrealString.h"
+#include "Internationalization/Internationalization.h"
 #include "Misc/Paths.h"
 #include "Modules/ModuleInterface.h"
 #include "Modules/ModuleManager.h"
 #include "IMediaModule.h"
-#include "IAndroidMediaModule.h"
-#include "AndroidMediaFactoryPrivate.h"
 #include "IMediaPlayerFactory.h"
+#include "UObject/NameTypes.h"
+
+#include "../../AndroidMedia/Public/IAndroidMediaModule.h"
 
 
 DEFINE_LOG_CATEGORY(LogAndroidMediaFactory);
@@ -26,7 +32,7 @@ public:
 
 	//~ IMediaPlayerFactory interface
 
-	virtual bool CanPlayUrl(const FString& Url, const IMediaOptions& Options, TArray<FText>* OutWarnings, TArray<FText>* OutErrors) const override
+	virtual bool CanPlayUrl(const FString& Url, const IMediaOptions* /*Options*/, TArray<FText>* /*OutWarnings*/, TArray<FText>* OutErrors) const override
 	{
 		FString Scheme;
 		FString Location;
@@ -71,10 +77,10 @@ public:
 		return true;
 	}
 
-	virtual TSharedPtr<IMediaPlayer, ESPMode::ThreadSafe> CreatePlayer() override
+	virtual TSharedPtr<IMediaPlayer, ESPMode::ThreadSafe> CreatePlayer(IMediaEventSink& EventSink) override
 	{
 		auto AndroidMediaModule = FModuleManager::LoadModulePtr<IAndroidMediaModule>("AndroidMedia");
-		return (AndroidMediaModule != nullptr) ? AndroidMediaModule->CreatePlayer() : nullptr;
+		return (AndroidMediaModule != nullptr) ? AndroidMediaModule->CreatePlayer(EventSink) : nullptr;
 	}
 
 	virtual FText GetDisplayName() const override
@@ -93,6 +99,13 @@ public:
 		return SupportedPlatforms;
 	}
 
+	virtual bool SupportsFeature(EMediaFeature Feature) const override
+	{
+		return ((Feature == EMediaFeature::AudioTracks) ||
+				(Feature == EMediaFeature::VideoSamples) ||
+				(Feature == EMediaFeature::VideoTracks));
+	}
+
 public:
 
 	//~ IModuleInterface interface
diff --git a/Engine/Plugins/Media/AndroidMedia/Source/AndroidMediaFactory/Private/AndroidMediaFactoryPrivate.h b/Engine/Plugins/Media/AndroidMedia/Source/AndroidMediaFactory/Private/AndroidMediaFactoryPrivate.h
--- a/Engine/Plugins/Media/AndroidMedia/Source/AndroidMediaFactory/Private/AndroidMediaFactoryPrivate.h
+++ b/Engine/Plugins/Media/AndroidMedia/Source/AndroidMediaFactory/Private/AndroidMediaFactoryPrivate.h
@@ -2,8 +2,8 @@
 
 #pragma once
 
-#include "CoreMinimal.h"
-
+#include "Logging/LogMacros.h"
 
 
+/** Log category for the AndroidMediaFactory module. */
 DECLARE_LOG_CATEGORY_EXTERN(LogAndroidMediaFactory, Log, All);
diff --git a/Engine/Plugins/Media/AvfMedia/AvfMedia.uplugin b/Engine/Plugins/Media/AvfMedia/AvfMedia.uplugin
--- a/Engine/Plugins/Media/AvfMedia/AvfMedia.uplugin
+++ b/Engine/Plugins/Media/AvfMedia/AvfMedia.uplugin
@@ -27,6 +27,11 @@
 			"Type" : "Editor",
 			"LoadingPhase" : "PostEngineInit"
 		},
+		{
+			"Name" : "AvfMediaFactory",
+			"Type" : "Editor",
+			"LoadingPhase" : "PostEngineInit"
+		},
 		{
 			"Name" : "AvfMediaFactory",
 			"Type" : "RuntimeNoCommandlet",
diff --git a/Engine/Plugins/Media/AvfMedia/Source/AvfMedia/AvfMedia.Build.cs b/Engine/Plugins/Media/AvfMedia/Source/AvfMedia/AvfMedia.Build.cs
--- a/Engine/Plugins/Media/AvfMedia/Source/AvfMedia/AvfMedia.Build.cs
+++ b/Engine/Plugins/Media/AvfMedia/Source/AvfMedia/AvfMedia.Build.cs
@@ -9,49 +9,45 @@ namespace UnrealBuildTool.Rules
 			DynamicallyLoadedModuleNames.AddRange(
 				new string[] {
 					"Media",
-				}
-			);
+				});
 
 			PrivateDependencyModuleNames.AddRange(
 				new string[] {
-                    "AvfMediaFactory",
+					"AvfMediaFactory",
 					"Core",
-                    "ShaderCore",
+					"ShaderCore",
+					"MediaUtils",
 					"RenderCore",
 					"RHI",
-                    "UtilityShaders",
-				}
-			);
+					"UtilityShaders",
+				});
 
-            PrivateIncludePathModuleNames.AddRange(
-                new string[] {
-                    "UtilityShaders",
-                }
-            );
+			PrivateIncludePathModuleNames.AddRange(
+				new string[] {
+					"UtilityShaders",
+				});
 
 			PrivateIncludePathModuleNames.AddRange(
 				new string[] {
 					"Media",
- 				}
-			);
+				});
 
 			PrivateIncludePaths.AddRange(
 				new string[] {
 					"AvfMedia/Private",
 					"AvfMedia/Private/Player",
+					"AvfMedia/Private/Shared",
 					"AvfMedia/Private/Tracks",
-				}
-			);
+				});
 
 			PublicFrameworks.AddRange(
 				new string[] {
 					"CoreMedia",
-                    "CoreVideo",
+					"CoreVideo",
 					"AVFoundation",
 					"AudioToolbox",
 					"QuartzCore"
-				}
-			);
+				});
 		}
 	}
 }
diff --git a/Engine/Plugins/Media/AvfMedia/Source/AvfMedia/Private/AvfMediaModule.cpp b/Engine/Plugins/Media/AvfMedia/Source/AvfMedia/Private/AvfMediaModule.cpp
--- a/Engine/Plugins/Media/AvfMedia/Source/AvfMedia/Private/AvfMediaModule.cpp
+++ b/Engine/Plugins/Media/AvfMedia/Source/AvfMedia/Private/AvfMediaModule.cpp
@@ -1,6 +1,9 @@
 // Copyright 1998-2017 Epic Games, Inc. All Rights Reserved.
 
 #include "AvfMediaPrivate.h"
+
+#include "Modules/ModuleManager.h"
+
 #include "AvfMediaPlayer.h"
 #include "IAvfMediaModule.h"
 
@@ -25,9 +28,9 @@ public:
 
 	//~ IAvfMediaModule interface
 
-	virtual TSharedPtr<IMediaPlayer, ESPMode::ThreadSafe> CreatePlayer() override
+	virtual TSharedPtr<IMediaPlayer, ESPMode::ThreadSafe> CreatePlayer(IMediaEventSink& EventSink) override
 	{
-		return MakeShareable(new FAvfMediaPlayer());
+		return MakeShared<FAvfMediaPlayer, ESPMode::ThreadSafe>(EventSink);
 	}
 
 public:
diff --git a/Engine/Plugins/Media/AvfMedia/Source/AvfMedia/Private/AvfMediaPrivate.h b/Engine/Plugins/Media/AvfMedia/Source/AvfMedia/Private/AvfMediaPrivate.h
--- a/Engine/Plugins/Media/AvfMedia/Source/AvfMedia/Private/AvfMediaPrivate.h
+++ b/Engine/Plugins/Media/AvfMedia/Source/AvfMedia/Private/AvfMediaPrivate.h
@@ -2,26 +2,7 @@
 
 #pragma once
 
-
-#include "Runtime/Core/Public/CoreMinimal.h"
-#include "Runtime/Core/Public/Modules/ModuleManager.h"
-#include "Runtime/Media/Public/IMediaAudioSink.h"
-#include "Runtime/Media/Public/IMediaModule.h"
-#include "Runtime/Media/Public/IMediaOverlaySink.h"
-#include "Runtime/Media/Public/IMediaPlayerFactory.h"
-#include "Runtime/Media/Public/IMediaTextureSink.h"
-
-#include "../../AvfMediaFactory/Public/AvfMediaSettings.h"
-
-#import <AVFoundation/AVFoundation.h>
-
-#if !PLATFORM_MAC || (MAC_OS_X_VERSION_MIN_REQUIRED >= 101100)
-#include <CoreVideo/CVMetalTextureCache.h>
-#endif
-
-#if PLATFORM_MAC
-	#include "CocoaThread.h"
-#endif
+#include "Logging/LogMacros.h"
 
 
 /** Log category for the AvfMedia module. */
diff --git a/Engine/Plugins/Media/AvfMedia/Source/AvfMedia/Private/Player/AvfMediaPlayer.cpp b/Engine/Plugins/Media/AvfMedia/Source/AvfMedia/Private/Player/AvfMediaPlayer.cpp
--- a/Engine/Plugins/Media/AvfMedia/Source/AvfMedia/Private/Player/AvfMediaPlayer.cpp
+++ b/Engine/Plugins/Media/AvfMedia/Source/AvfMedia/Private/Player/AvfMediaPlayer.cpp
@@ -1,13 +1,27 @@
 // Copyright 1998-2017 Epic Games, Inc. All Rights Reserved.
 
 #include "AvfMediaPlayer.h"
+#include "AvfMediaPrivate.h"
+
 #include "HAL/PlatformProcess.h"
+#include "IMediaEventSink.h"
+#include "MediaSamples.h"
 #include "Misc/CommandLine.h"
 #include "Misc/Paths.h"
-#if !PLATFORM_MAC
-#include "IOS/IOSAsyncTask.h"
+
+#if PLATFORM_MAC
+	#include "Mac/CocoaThread.h"
+#else
+	#include "IOS/IOSAsyncTask.h"
 #endif
 
+#include "AvfMediaTracks.h"
+#include "AvfMediaUtils.h"
+
+
+/* FAVPlayerDelegate
+ *****************************************************************************/
+
 /**
  * Cocoa class that can help us with reading player item information.
  */
@@ -32,54 +46,6 @@
 
 @end
 
-#if !PLATFORM_MAC
-static FString ConvertToIOSPath(const FString& Filename, bool bForWrite)
-{
-	FString Result = Filename;
-	if (Result.Contains(TEXT("/OnDemandResources/")))
-	{
-		return Result;
-	}
-	
-	Result.ReplaceInline(TEXT("../"), TEXT(""));
-	Result.ReplaceInline(TEXT(".."), TEXT(""));
-	Result.ReplaceInline(FPlatformProcess::BaseDir(), TEXT(""));
-	
-	if(bForWrite)
-	{
-		static FString WritePathBase = FString([NSSearchPathForDirectoriesInDomains(NSDocumentDirectory, NSUserDomainMask, YES) objectAtIndex:0]) + TEXT("/");
-		return WritePathBase + Result;
-	}
-	else
-	{
-		// if filehostip exists in the command line, cook on the fly read path should be used
-		FString Value;
-		// Cache this value as the command line doesn't change...
-		static bool bHasHostIP = FParse::Value(FCommandLine::Get(), TEXT("filehostip"), Value) || FParse::Value(FCommandLine::Get(), TEXT("streaminghostip"), Value);
-		static bool bIsIterative = FParse::Value(FCommandLine::Get(), TEXT("iterative"), Value);
-		if (bHasHostIP)
-		{
-			static FString ReadPathBase = FString([NSSearchPathForDirectoriesInDomains(NSDocumentDirectory, NSUserDomainMask, YES) objectAtIndex:0]) + TEXT("/");
-			return ReadPathBase + Result;
-		}
-		else if (bIsIterative)
-		{
-			static FString ReadPathBase = FString([NSSearchPathForDirectoriesInDomains(NSDocumentDirectory, NSUserDomainMask, YES) objectAtIndex:0]) + TEXT("/");
-			return ReadPathBase + Result.ToLower();
-		}
-		else
-		{
-			static FString ReadPathBase = FString([[NSBundle mainBundle] bundlePath]) + TEXT("/cookeddata/");
-			return ReadPathBase + Result.ToLower();
-		}
-	}
-	
-	return Result;
-}
-#endif
-
-/* FAVPlayerDelegate implementation
- *****************************************************************************/
 
 @implementation FAVPlayerDelegate
 @synthesize MediaPlayer;
@@ -107,7 +73,7 @@ static FString ConvertToIOSPath(const FString& Filename, bool bForWrite)
 		if (object == (id)context)
 		{
 			AVPlayerItemStatus Status = ((AVPlayerItem*)object).status;
-			MediaPlayer->HandleStatusNotification(Status);
+			MediaPlayer->OnStatusNotification(Status);
 		}
 	}
 }
@@ -118,9 +84,10 @@ static FString ConvertToIOSPath(const FString& Filename, bool bForWrite)
 	[super dealloc];
 }
 
+
 -(void)playerItemPlaybackEndReached:(NSNotification*)Notification
 {
-	MediaPlayer->HandleDidReachEnd();
+	MediaPlayer->OnEndReached();
 }
 
 @end
@@ -129,45 +96,80 @@ static FString ConvertToIOSPath(const FString& Filename, bool bForWrite)
 /* FAvfMediaPlayer structors
  *****************************************************************************/
 
-FAvfMediaPlayer::FAvfMediaPlayer()
+FAvfMediaPlayer::FAvfMediaPlayer(IMediaEventSink& InEventSink)
+	: EventSink(InEventSink)
 {
+	CurrentRate = 0.0f;
+	CurrentState = EMediaState::Closed;
     CurrentTime = FTimespan::Zero();
-    Duration = FTimespan::Zero();
+
+	Duration = FTimespan::Zero();
 	MediaUrl = FString();
 	ShouldLoop = false;
     
 	MediaHelper = nil;
     MediaPlayer = nil;
 	PlayerItem = nil;
-	
-	CurrentRate = 0.0f;
-	
-	State = EMediaState::Closed;
-	
+		
 	bPrerolled = false;
+
+	Samples = new FMediaSamples;
+	Tracks = new FAvfMediaTracks(*Samples);
 }
 
+
 FAvfMediaPlayer::~FAvfMediaPlayer()
 {
 	Close();
+
+	delete Samples;
+	Samples = nullptr;
 }
 
-void FAvfMediaPlayer::HandleStatusNotification(AVPlayerItemStatus Status)
+
+/* FAvfMediaPlayer interface
+ *****************************************************************************/
+
+void FAvfMediaPlayer::OnEndReached()
+{
+	if (ShouldLoop)
+	{
+		PlayerTasks.Enqueue([=]() {
+			EventSink.ReceiveMediaEvent(EMediaEvent::PlaybackEndReached);
+			Seek(FTimespan::FromSeconds(0.0f));
+			SetRate(CurrentRate);
+		});
+	}
+	else
+	{
+		CurrentState = EMediaState::Paused;
+		CurrentRate = 0.0f;
+
+		PlayerTasks.Enqueue([=]() {
+			Seek(FTimespan::FromSeconds(0.0f));
+			EventSink.ReceiveMediaEvent(EMediaEvent::PlaybackEndReached);
+			EventSink.ReceiveMediaEvent(EMediaEvent::PlaybackSuspended);
+		});
+	}
+}
+
+
+void FAvfMediaPlayer::OnStatusNotification(AVPlayerItemStatus Status)
 {
 	switch(Status)
 	{
 		case AVPlayerItemStatusReadyToPlay:
 		{
-			if (Duration == 0.0f || State == EMediaState::Closed)
+			if (Duration == 0.0f || CurrentState == EMediaState::Closed)
 			{
-				Tracks.Initialize(PlayerItem, Info);
+				Tracks->Initialize(PlayerItem, Info);
 			 
 				PlayerTasks.Enqueue([=]() {
-					MediaEvent.Broadcast(EMediaEvent::TracksChanged);
+					EventSink.ReceiveMediaEvent(EMediaEvent::TracksChanged);
 				});
 				
 				Duration = FTimespan::FromSeconds(CMTimeGetSeconds(PlayerItem.asset.duration));
-				State = (State == EMediaState::Closed) ? EMediaState::Preparing : State;
+				CurrentState = (CurrentState == EMediaState::Closed) ? EMediaState::Preparing : CurrentState;
 				
 				PlayerTasks.Enqueue([=]() {
 					if (!bPrerolled)
@@ -178,8 +180,9 @@ void FAvfMediaPlayer::HandleStatusNotification(AVPlayerItemStatus Status)
 							bPrerolled = true;
 							if (bFinished)
 							{
+								CurrentState = EMediaState::Stopped;
 								PlayerTasks.Enqueue([=]() {
-									MediaEvent.Broadcast(EMediaEvent::MediaOpened);
+									EventSink.ReceiveMediaEvent(EMediaEvent::MediaOpened);
 									if (CurrentTime != FTimespan::Zero())
 									{
 										Seek(CurrentTime);
@@ -192,9 +195,9 @@ void FAvfMediaPlayer::HandleStatusNotification(AVPlayerItemStatus Status)
 							}
 							else
 							{
-								State = EMediaState::Error;
+								CurrentState = EMediaState::Error;
 								PlayerTasks.Enqueue([=]() {
-									MediaEvent.Broadcast(EMediaEvent::MediaOpenFailed);
+									EventSink.ReceiveMediaEvent(EMediaEvent::MediaOpenFailed);
 								});
 							}
 						}];
@@ -205,18 +208,18 @@ void FAvfMediaPlayer::HandleStatusNotification(AVPlayerItemStatus Status)
 		}
 		case AVPlayerItemStatusFailed:
 		{
-			if (Duration == 0.0f || State == EMediaState::Closed)
+			if (Duration == 0.0f || CurrentState == EMediaState::Closed)
 			{
-				State = EMediaState::Error;
+				CurrentState = EMediaState::Error;
 				PlayerTasks.Enqueue([=]() {
-					MediaEvent.Broadcast(EMediaEvent::MediaOpenFailed);
+					EventSink.ReceiveMediaEvent(EMediaEvent::MediaOpenFailed);
 				});
 			}
 			else
 			{
-				State = EMediaState::Error;
+				CurrentState = EMediaState::Error;
 				PlayerTasks.Enqueue([=]() {
-					MediaEvent.Broadcast(EMediaEvent::PlaybackSuspended);
+					EventSink.ReceiveMediaEvent(EMediaEvent::PlaybackSuspended);
 				});
 			}
 			break;
@@ -228,224 +231,63 @@ void FAvfMediaPlayer::HandleStatusNotification(AVPlayerItemStatus Status)
 }
 
 
-void FAvfMediaPlayer::HandleDidReachEnd()
-{
-	if (ShouldLoop)
-	{
-		PlayerTasks.Enqueue([=]() {
-			MediaEvent.Broadcast(EMediaEvent::PlaybackEndReached);
-			Seek(FTimespan::FromSeconds(0.0f));
-			SetRate(CurrentRate);
-		});
-	}
-	else
-	{
-		State = EMediaState::Paused;
-		CurrentRate = 0.0f;
-
-		PlayerTasks.Enqueue([=]() {
-			Seek(FTimespan::FromSeconds(0.0f));
-			MediaEvent.Broadcast(EMediaEvent::PlaybackEndReached);
-			MediaEvent.Broadcast(EMediaEvent::PlaybackSuspended);
-		});
-	}
-}
-
-
-/* IMediaControls interface
+/* IMediaPlayer interface
  *****************************************************************************/
 
-FTimespan FAvfMediaPlayer::GetDuration() const
-{
-	return Duration;
-}
-
-
-float FAvfMediaPlayer::GetRate() const
-{
-	return CurrentRate;
-}
-
-
-EMediaState FAvfMediaPlayer::GetState() const
-{
-	return State;
-}
-
-
-TRange<float> FAvfMediaPlayer::GetSupportedRates(EMediaPlaybackDirections Direction, bool Unthinned) const
+void FAvfMediaPlayer::Close()
 {
-	if (Direction == EMediaPlaybackDirections::Forward)
+	if (CurrentState == EMediaState::Closed)
 	{
-		float Max = PlayerItem.canPlayFastForward ? 8.0f : 1.0f;
-		float Min = 0.0f;
-		return TRange<float>(Min, Max);
+		return;
 	}
-	else
+
+	CurrentTime = 0;
+	MediaUrl = FString();
+	
+	if (PlayerItem != nil)
 	{
-		if (PlayerItem.canPlayReverse)
-		{
-			float Max = PlayerItem.canPlayFastReverse ? 8.0f : 1.0f;
-			float Min = 0.0f;
-			return TRange<float>(Min, Max);
-		}
-		else
+		if (MediaHelper != nil)
 		{
-			return TRange<float>::Empty();
+			[[NSNotificationCenter defaultCenter] removeObserver:MediaHelper name:AVPlayerItemDidPlayToEndTimeNotification object:PlayerItem];
 		}
-	}
-}
-
 
-FTimespan FAvfMediaPlayer::GetTime() const
-{
-	FTimespan Time = CurrentTime;
-	if (MediaPlayer)
-	{
-		CMTime Current = [MediaPlayer currentTime];
-		FTimespan DiplayTime = FTimespan::FromSeconds(CMTimeGetSeconds(Current));
-		Time = FMath::Min(DiplayTime, Duration);
+		[PlayerItem removeObserver:MediaHelper forKeyPath:@"status"];
+		[PlayerItem release];
+		PlayerItem = nil;
 	}
 	
-	return Time;
-}
-
-
-bool FAvfMediaPlayer::IsLooping() const
-{
-	return ShouldLoop;
-}
-
-
-bool FAvfMediaPlayer::Seek(const FTimespan& Time)
-{
-	CurrentTime = Time;
-	
-	if (bPrerolled)
+	if (MediaHelper != nil)
 	{
-		Tracks.Seek(Time);
-		
-		double TotalSeconds = Time.GetTotalSeconds();
-		CMTime CurrentTimeInSeconds = CMTimeMakeWithSeconds(TotalSeconds, 1000);
-		
-		static CMTime Tolerance = CMTimeMakeWithSeconds(0.01, 1000);
-		[MediaPlayer seekToTime:CurrentTimeInSeconds toleranceBefore:Tolerance toleranceAfter:Tolerance];
+		[MediaHelper release];
+		MediaHelper = nil;
 	}
 
-	return true;
-}
-
-
-bool FAvfMediaPlayer::SetLooping(bool Looping)
-{
-	ShouldLoop = Looping;
-	
-	if (ShouldLoop)
+	if (MediaPlayer != nil)
 	{
-		MediaPlayer.actionAtItemEnd = AVPlayerActionAtItemEndNone;
+		// If we don't remove the current player item then the retain count is > 1 for the MediaPlayer then on it's release then the MetalPlayer stays around forever
+		[MediaPlayer replaceCurrentItemWithPlayerItem:nil];
+		[MediaPlayer release];
+		MediaPlayer = nil;
 	}
-	else
-	{
-		MediaPlayer.actionAtItemEnd = AVPlayerActionAtItemEndPause;
-	}
-
-	return true;
-}
-
-
-bool FAvfMediaPlayer::SetRate(float Rate)
-{
-	CurrentRate = Rate;
 	
-	if (bPrerolled)
-	{
-		[MediaPlayer setRate : CurrentRate];
+	Tracks->Reset();
+	EventSink.ReceiveMediaEvent(EMediaEvent::TracksChanged);
 
-		Tracks.SetRate(Rate);
+	CurrentState = EMediaState::Closed;
+	Duration = CurrentTime = FTimespan::Zero();
+	Info.Empty();
+	
+	EventSink.ReceiveMediaEvent(EMediaEvent::MediaClosed);
 		
-		if (FMath::IsNearlyZero(Rate))
-		{
-			State = EMediaState::Paused;
-			MediaEvent.Broadcast(EMediaEvent::PlaybackSuspended);
-		}
-		else
-		{
-			State = EMediaState::Playing;
-			MediaEvent.Broadcast(EMediaEvent::PlaybackResumed);
-		}
-	}
-
-	return true;
-}
-
-
-bool FAvfMediaPlayer::SupportsRate(float Rate, bool Unthinned) const
-{
-	return GetSupportedRates(Rate >= 0.0f ? EMediaPlaybackDirections::Forward : EMediaPlaybackDirections::Reverse, Unthinned).Contains(FMath::Abs(Rate));
-}
-
-
-bool FAvfMediaPlayer::SupportsScrubbing() const
-{
-	return true;
-}
-
-
-bool FAvfMediaPlayer::SupportsSeeking() const
-{
-	return true;
+	bPrerolled = false;
+	
+	CurrentRate = 0.f;
 }
 
 
-/* IMediaPlayer interface
- *****************************************************************************/
-
-void FAvfMediaPlayer::Close()
+IMediaCache& FAvfMediaPlayer::GetCache()
 {
-	if (GetState() > EMediaState::Closed)
-	{
-		CurrentTime = 0;
-		MediaUrl = FString();
-	
-		if (PlayerItem != nil)
-		{
-			if(MediaHelper != nil)
-			{
-				[[NSNotificationCenter defaultCenter] removeObserver:MediaHelper name:AVPlayerItemDidPlayToEndTimeNotification object:PlayerItem];
-                [PlayerItem removeObserver:MediaHelper forKeyPath:@"status" context:PlayerItem];
-			}
-			
-			[PlayerItem release];
-			PlayerItem = nil;
-		}
-	
-		if (MediaHelper != nil)
-		{
-			[MediaHelper release];
-			MediaHelper = nil;
-		}
-
-		if (MediaPlayer != nil)
-		{
-			[MediaPlayer release];
-			MediaPlayer = nil;
-		}
-	
-		Tracks.Reset();
-		MediaEvent.Broadcast(EMediaEvent::TracksChanged);
-	
-		Duration = CurrentTime = FTimespan::Zero();
-
-		Info.Empty();
-	
-		MediaEvent.Broadcast(EMediaEvent::MediaClosed);
-        
-        State = EMediaState::Closed;
-		
-		bPrerolled = false;
-
-		State = EMediaState::Closed;
-	}
+	return *this;
 }
 
 
@@ -468,9 +310,9 @@ FName FAvfMediaPlayer::GetName() const
 }
 
 
-IMediaOutput& FAvfMediaPlayer::GetOutput()
+IMediaSamples& FAvfMediaPlayer::GetSamples()
 {
-	return Tracks;
+	return *Samples;
 }
 
 
@@ -478,7 +320,7 @@ FString FAvfMediaPlayer::GetStats() const
 {
 	FString Result;
 
-	Tracks.AppendStats(Result);
+	Tracks->AppendStats(Result);
 
 	return Result;
 }
@@ -486,7 +328,7 @@ FString FAvfMediaPlayer::GetStats() const
 
 IMediaTracks& FAvfMediaPlayer::GetTracks()
 {
-	return Tracks;
+	return *Tracks;
 }
 
 
@@ -496,12 +338,19 @@ FString FAvfMediaPlayer::GetUrl() const
 }
 
 
-bool FAvfMediaPlayer::Open(const FString& Url, const IMediaOptions& Options)
+IMediaView& FAvfMediaPlayer::GetView()
+{
+	return *this;
+}
+
+
+bool FAvfMediaPlayer::Open(const FString& Url, const IMediaOptions* /*Options*/)
 {
 	Close();
 
 	NSURL* nsMediaUrl = nil;
 	FString Path;
+
 	if (Url.StartsWith(TEXT("file://")))
 	{
 		// Media Framework doesn't percent encode the URL, so the path portion is just a native file path.
@@ -525,11 +374,11 @@ bool FAvfMediaPlayer::Open(const FString& Url, const IMediaOptions& Options)
 
 	// On non-Mac Apple OSes the path is:
 	//	a) case-sensitive
-	//	b) relative to the cookeddata directory, not the notional GameContentDirectory which is 'virtual' and resolved by the FIOSPlatformFile calls
+	//	b) relative to the 'cookeddata' directory, not the notional GameContentDirectory which is 'virtual' and resolved by the FIOSPlatformFile calls
 #if !PLATFORM_MAC
 	if ([[nsMediaUrl scheme] isEqualToString:@"file"])
 	{
-		FString FullPath = ConvertToIOSPath(Path, false);
+		FString FullPath = AvfMedia::ConvertToIOSPath(Path, false);
 		nsMediaUrl = [NSURL fileURLWithPath: FullPath.GetNSString() isDirectory:NO];
 	}
 #endif
@@ -551,7 +400,7 @@ bool FAvfMediaPlayer::Open(const FString& Url, const IMediaOptions& Options)
 	MediaHelper = [[FAVPlayerDelegate alloc] initWithMediaPlayer:this];
 	check(MediaHelper != nil);
 
-	PlayerItem = [AVPlayerItem playerItemWithURL : nsMediaUrl];
+	PlayerItem = [[AVPlayerItem playerItemWithURL:nsMediaUrl] retain];
 
 	if (PlayerItem == nil)
 	{
@@ -560,9 +409,9 @@ bool FAvfMediaPlayer::Open(const FString& Url, const IMediaOptions& Options)
 		return false;
 	}
 
-	// load tracks
-	[PlayerItem retain];
+	CurrentState = EMediaState::Preparing;
 
+	// load tracks
 	[[PlayerItem asset] loadValuesAsynchronouslyForKeys:@[@"tracks"] completionHandler:^
 	{
 		NSError* Error = nil;
@@ -574,7 +423,7 @@ bool FAvfMediaPlayer::Open(const FString& Url, const IMediaOptions& Options)
 			// File movies will be ready now
 			if (PlayerItem.status == AVPlayerItemStatusReadyToPlay)
 			{
-				HandleStatusNotification(PlayerItem.status);
+				OnStatusNotification(PlayerItem.status);
 			}
 			
 			// Streamed movies might not be and we want to know if it ever fails.
@@ -582,7 +431,7 @@ bool FAvfMediaPlayer::Open(const FString& Url, const IMediaOptions& Options)
 		}
 		else if (Error != nullptr)
 		{
-			State = EMediaState::Error;
+			CurrentState = EMediaState::Error;
 			
 			NSDictionary *userInfo = [Error userInfo];
 			NSString *errstr = [[userInfo objectForKey : NSUnderlyingErrorKey] localizedDescription];
@@ -590,7 +439,7 @@ bool FAvfMediaPlayer::Open(const FString& Url, const IMediaOptions& Options)
 			UE_LOG(LogAvfMedia, Warning, TEXT("Failed to load video tracks. [%s]"), *FString(errstr));
 	 
 			PlayerTasks.Enqueue([=]() {
-				MediaEvent.Broadcast(EMediaEvent::MediaOpenFailed);
+				EventSink.ReceiveMediaEvent(EMediaEvent::MediaOpenFailed);
 			});
 		}
 	}];
@@ -605,13 +454,37 @@ bool FAvfMediaPlayer::Open(const FString& Url, const IMediaOptions& Options)
 }
 
 
-bool FAvfMediaPlayer::Open(const TSharedRef<FArchive, ESPMode::ThreadSafe>& Archive, const FString& OriginalUrl, const IMediaOptions& Options)
+bool FAvfMediaPlayer::Open(const TSharedRef<FArchive, ESPMode::ThreadSafe>& /*Archive*/, const FString& /*OriginalUrl*/, const IMediaOptions* /*Options*/)
 {
 	return false; // not supported
 }
 
 
-void FAvfMediaPlayer::TickPlayer(float DeltaTime)
+void FAvfMediaPlayer::TickAudio()
+{
+	if ((CurrentState > EMediaState::Error) && (Duration > 0.0f))
+	{
+		Tracks->ProcessAudio();
+		if(MediaPlayer != nil)
+		{
+			CMTime Current = [MediaPlayer currentTime];
+			FTimespan DiplayTime = FTimespan::FromSeconds(CMTimeGetSeconds(Current));
+			CurrentTime = FMath::Min(DiplayTime, Duration);
+		}
+	}
+}
+
+
+void FAvfMediaPlayer::TickFetch(FTimespan DeltaTime)
+{
+	if ((CurrentState > EMediaState::Error) && (Duration > 0.0f))
+	{
+		Tracks->ProcessVideo();
+	}
+}
+
+
+void FAvfMediaPlayer::TickInput(FTimespan DeltaTime)
 {
 	// process deferred tasks
 	TFunction<void()> Task;
@@ -620,16 +493,144 @@ void FAvfMediaPlayer::TickPlayer(float DeltaTime)
 	{
 		Task();
 	}
+}
+
+
+/* IMediaControls interface
+ *****************************************************************************/
+
+bool FAvfMediaPlayer::CanControl(EMediaControl Control) const
+{
+	if (!bPrerolled)
+	{
+		return false;
+	}
+
+	if (Control == EMediaControl::Pause)
+	{
+		return (CurrentState == EMediaState::Playing);
+	}
+
+	if (Control == EMediaControl::Resume)
+	{
+		return (CurrentState != EMediaState::Playing);
+	}
+
+	if ((Control == EMediaControl::Scrub) || (Control == EMediaControl::Seek))
+	{
+		return true;
+	}
+
+	return false;
+}
+
+
+FTimespan FAvfMediaPlayer::GetDuration() const
+{
+	return Duration;
+}
+
+
+float FAvfMediaPlayer::GetRate() const
+{
+	return CurrentRate;
+}
+
+
+EMediaState FAvfMediaPlayer::GetState() const
+{
+	return CurrentState;
+}
+
+
+EMediaStatus FAvfMediaPlayer::GetStatus() const
+{
+	return EMediaStatus::None;
+}
+
+
+TRangeSet<float> FAvfMediaPlayer::GetSupportedRates(EMediaRateThinning Thinning) const
+{
+	TRangeSet<float> Result;
+
+	Result.Add(TRange<float>(PlayerItem.canPlayFastReverse ? -8.0f : -1.0f, 0.0f));
+	Result.Add(TRange<float>(0.0f, PlayerItem.canPlayFastForward ? 8.0f : 0.0f));
+
+	return Result;
+}
+
+
+FTimespan FAvfMediaPlayer::GetTime() const
+{
+	return CurrentTime;
+}
+
+
+bool FAvfMediaPlayer::IsLooping() const
+{
+	return ShouldLoop;
+}
+
+
+bool FAvfMediaPlayer::Seek(const FTimespan& Time)
+{
+	CurrentTime = Time;
+	
+	if (bPrerolled)
+	{
+		Tracks->Seek(Time);
+		
+		double TotalSeconds = Time.GetTotalSeconds();
+		CMTime CurrentTimeInSeconds = CMTimeMakeWithSeconds(TotalSeconds, 1000);
+		
+		static CMTime Tolerance = CMTimeMakeWithSeconds(0.01, 1000);
+		[MediaPlayer seekToTime:CurrentTimeInSeconds toleranceBefore:Tolerance toleranceAfter:Tolerance];
+	}
+
+	EventSink.ReceiveMediaEvent(EMediaEvent::SeekCompleted);
+
+	return true;
+}
+
 
-	// tick tracks (for audio output)
-	if ((GetState() > EMediaState::Error) && (Duration > 0.0f))
+bool FAvfMediaPlayer::SetLooping(bool Looping)
+{
+	ShouldLoop = Looping;
+	
+	if (ShouldLoop)
 	{
-		Tracks.Tick(DeltaTime);
+		MediaPlayer.actionAtItemEnd = AVPlayerActionAtItemEndNone;
 	}
+	else
+	{
+		MediaPlayer.actionAtItemEnd = AVPlayerActionAtItemEndPause;
+	}
+
+	return true;
 }
 
 
-void FAvfMediaPlayer::TickVideo(float DeltaTime)
+bool FAvfMediaPlayer::SetRate(float Rate)
 {
-	// do nothing
+	CurrentRate = Rate;
+	
+	if (bPrerolled)
+	{
+		[MediaPlayer setRate : CurrentRate];
+
+		Tracks->SetRate(Rate);
+		
+		if (FMath::IsNearlyZero(Rate))
+		{
+			CurrentState = EMediaState::Paused;
+			EventSink.ReceiveMediaEvent(EMediaEvent::PlaybackSuspended);
+		}
+		else
+		{
+			CurrentState = EMediaState::Playing;
+			EventSink.ReceiveMediaEvent(EMediaEvent::PlaybackResumed);
+		}
+	}
+
+	return true;
 }
diff --git a/Engine/Plugins/Media/AvfMedia/Source/AvfMedia/Private/Player/AvfMediaPlayer.h b/Engine/Plugins/Media/AvfMedia/Source/AvfMedia/Private/Player/AvfMediaPlayer.h
--- a/Engine/Plugins/Media/AvfMedia/Source/AvfMedia/Private/Player/AvfMediaPlayer.h
+++ b/Engine/Plugins/Media/AvfMedia/Source/AvfMedia/Private/Player/AvfMediaPlayer.h
@@ -2,14 +2,17 @@
 
 #pragma once
 
-#include "../AvfMediaPrivate.h"
-#include "CoreMinimal.h"
-#include "AvfMediaTracks.h"
+#include "Containers/Queue.h"
+#include "IMediaCache.h"
 #include "IMediaControls.h"
 #include "IMediaPlayer.h"
-#include "Containers/Queue.h"
-#include "Containers/Ticker.h"
+#include "IMediaView.h"
+
+#import <AVFoundation/AVFoundation.h>
 
+class FAvfMediaTracks;
+class FMediaSamples;
+class IMediaEventSink;
 
 @class AVPlayer;
 @class AVPlayerItem;
@@ -20,80 +23,87 @@
  * Implements a media player using the AV framework.
  */
 class FAvfMediaPlayer
-	: public IMediaControls
-	, public IMediaPlayer
+	: public IMediaPlayer
+	, protected IMediaCache
+	, protected IMediaControls
+	, protected IMediaView
 {
 public:
 
-	/** Default constructor. */
-	FAvfMediaPlayer();
+	/**
+	 * Create and initialize a new instance.
+	 *
+	 * @param InEventSink The object that receives media events from this player.
+	 */
+	FAvfMediaPlayer(IMediaEventSink& InEventSink);
 
 	/** Destructor. */
 	~FAvfMediaPlayer();
 
-	/** Called by the delegate whenever the player item status changes. */
-	void HandleStatusNotification(AVPlayerItemStatus Status);
-
-	/** Called by the delegate when the playback reaches the end. */
-	void HandleDidReachEnd();
-
 public:
 
-	//~ IMediaControls interface
+	/** Called by the delegate when the playback reaches the end. */
+	void OnEndReached();
 
-	virtual FTimespan GetDuration() const override;
-	virtual float GetRate() const override;
-	virtual EMediaState GetState() const override;
-	virtual TRange<float> GetSupportedRates(EMediaPlaybackDirections Direction, bool Unthinned) const override;
-	virtual FTimespan GetTime() const override;
-	virtual bool IsLooping() const override;
-	virtual bool Seek(const FTimespan& Time) override;
-	virtual bool SetLooping(bool Looping) override;
-	virtual bool SetRate(float Rate) override;
-	virtual bool SupportsRate(float Rate, bool Unthinned) const override;
-	virtual bool SupportsScrubbing() const override;
-	virtual bool SupportsSeeking() const override;
+	/** Called by the delegate whenever the player item status changes. */
+	void OnStatusNotification(AVPlayerItemStatus Status);
 
 public:
 
 	//~ IMediaPlayer interface
 
 	virtual void Close() override;
+	virtual IMediaCache& GetCache() override;
 	virtual IMediaControls& GetControls() override;
 	virtual FString GetInfo() const override;
 	virtual FName GetName() const override;
-	virtual IMediaOutput& GetOutput() override;
+	virtual IMediaSamples& GetSamples() override;
 	virtual FString GetStats() const override;
 	virtual IMediaTracks& GetTracks() override;
 	virtual FString GetUrl() const override;
-	virtual bool Open(const FString& Url, const IMediaOptions& Options) override;
-	virtual bool Open(const TSharedRef<FArchive, ESPMode::ThreadSafe>& Archive, const FString& OriginalUrl, const IMediaOptions& Options) override;
-	virtual void TickPlayer(float DeltaTime) override;
-	virtual void TickVideo(float DeltaTime) override;
+	virtual IMediaView& GetView() override;
+	virtual bool Open(const FString& Url, const IMediaOptions* Options) override;
+	virtual bool Open(const TSharedRef<FArchive, ESPMode::ThreadSafe>& Archive, const FString& OriginalUrl, const IMediaOptions* Options) override;
+	virtual void TickAudio() override;
+	virtual void TickFetch(FTimespan DeltaTime) override;
+	virtual void TickInput(FTimespan DeltaTime) override;
 
-	DECLARE_DERIVED_EVENT(FVlcMediaPlayer, IMediaPlayer::FOnMediaEvent, FOnMediaEvent);
-	virtual FOnMediaEvent& OnMediaEvent() override
-	{
-		return MediaEvent;
-	}
+protected:
+
+	//~ IMediaControls interface
+
+	virtual bool CanControl(EMediaControl Control) const override;
+	virtual FTimespan GetDuration() const override;
+	virtual float GetRate() const override;
+	virtual EMediaState GetState() const override;
+	virtual EMediaStatus GetStatus() const override;
+	virtual TRangeSet<float> GetSupportedRates(EMediaRateThinning Thinning) const override;
+	virtual FTimespan GetTime() const override;
+	virtual bool IsLooping() const override;
+	virtual bool Seek(const FTimespan& Time) override;
+	virtual bool SetLooping(bool Looping) override;
+	virtual bool SetRate(float Rate) override;
 
 private:
 
 	/** The current playback rate. */
 	float CurrentRate;
 
+	/** Media playback state. */
+	EMediaState CurrentState;
+
 	/** The current time of the playback. */
 	FTimespan CurrentTime;
 
 	/** The duration of the media. */
     FTimespan Duration;
 
+	/** The media event handler. */
+	IMediaEventSink& EventSink;
+
 	/** Media information string. */
 	FString Info;
 
-	/** Holds an event delegate that is invoked when a media event occurred. */
-	FOnMediaEvent MediaEvent;
-
 	/** Cocoa helper object we can use to keep track of ns property changes in our media items */
 	FAVPlayerDelegate* MediaHelper;
     
@@ -109,14 +119,17 @@ private:
 	/** Tasks to be executed on the player thread. */
 	TQueue<TFunction<void()>> PlayerTasks;
 
+	/** The media sample queue. */
+	FMediaSamples* Samples;
+
 	/** Should the video loop to the beginning at completion */
     bool ShouldLoop;
 
 	/** The media track collection. */
-	FAvfMediaTracks Tracks;
-	
-	/** Media playback state. */
-	EMediaState State;
+	FAvfMediaTracks* Tracks;
 	
 	bool bPrerolled;
+
+	/** Mutex to ensure thread-safe access */
+	FCriticalSection CriticalSection;
 };
