CIS fix.

diff --git a/Engine/Plugins/Media/AndroidMedia/Source/AndroidMedia/Private/Player/AndroidMediaPlayer.cpp b/Engine/Plugins/Media/AndroidMedia/Source/AndroidMedia/Private/Player/AndroidMediaPlayer.cpp
--- a/Engine/Plugins/Media/AndroidMedia/Source/AndroidMedia/Private/Player/AndroidMediaPlayer.cpp
+++ b/Engine/Plugins/Media/AndroidMedia/Source/AndroidMedia/Private/Player/AndroidMediaPlayer.cpp
@@ -275,7 +275,7 @@ bool FAndroidMediaPlayer::Open(const FString& Url, const IMediaOptions* /*Option
 			}
 
 			// is it a simple case (can just use file datasource)?
-			if (FileEntry.CompressionMethod == COMPRESS_None && !FileEntry.IsEncrypted())
+			if (FileEntry.CompressionMethodIndex == 0 && !FileEntry.IsEncrypted())
 			{
 				FString PakFilename = PakFile->GetFilename();
 				int64 PakHeaderSize = FileEntry.GetSerializedSize(PakFile->GetInfo().Version);
