Check for error with URL based Android movie playback

#ue4
#android

diff --git a/Engine/Plugins/Media/AndroidMedia/Source/AndroidMedia/Private/Player/AndroidMediaPlayer.cpp b/Engine/Plugins/Media/AndroidMedia/Source/AndroidMedia/Private/Player/AndroidMediaPlayer.cpp
--- a/Engine/Plugins/Media/AndroidMedia/Source/AndroidMedia/Private/Player/AndroidMediaPlayer.cpp
+++ b/Engine/Plugins/Media/AndroidMedia/Source/AndroidMedia/Private/Player/AndroidMediaPlayer.cpp
@@ -291,7 +291,11 @@ bool FAndroidMediaPlayer::Open(const FString& Url, const IMediaOptions* /*Option
 	else
 	{
 		// open remote media
-		JavaMediaPlayer->SetDataSource(Url);
+		if (!JavaMediaPlayer->SetDataSource(Url))
+		{
+			UE_LOG(LogAndroidMedia, Warning, TEXT("Failed to set data source for URL %s"), *Url);
+			return false;
+		}
 	}
 
 	// prepare media
