UE4.22 - Android compile fix
#rnx

diff --git a/Engine/Plugins/Media/AndroidMedia/Source/AndroidMedia/Private/Player/AndroidMediaPlayer.cpp b/Engine/Plugins/Media/AndroidMedia/Source/AndroidMedia/Private/Player/AndroidMediaPlayer.cpp
--- a/Engine/Plugins/Media/AndroidMedia/Source/AndroidMedia/Private/Player/AndroidMediaPlayer.cpp
+++ b/Engine/Plugins/Media/AndroidMedia/Source/AndroidMedia/Private/Player/AndroidMediaPlayer.cpp
@@ -78,14 +78,14 @@ FAndroidMediaPlayer::~FAndroidMediaPlayer()
 
 			FReleaseVideoResourcesParams ReleaseVideoResourcesParams = { VideoTexture, PlayerGuid };
 
-			ENQUEUE_UNIQUE_RENDER_COMMAND_ONEPARAMETER(AndroidMediaPlayerWriteVideoSample,
-				FReleaseVideoResourcesParams, Params, ReleaseVideoResourcesParams,
+			ENQUEUE_RENDER_COMMAND(AndroidMediaPlayerWriteVideoSample)(
+				[ReleaseVideoResourcesParams](FRHICommandListImmediate& RHICmdList)
 				{
 					#if ANDROIDMEDIAPLAYER_USE_NATIVELOGGING
-						FPlatformMisc::LowLevelOutputDebugStringf(TEXT("~FAndroidMediaPlayer: Unregister Guid: %s"), *Params.PlayerGuid.ToString());
+						FPlatformMisc::LowLevelOutputDebugStringf(TEXT("~FAndroidMediaPlayer: Unregister Guid: %s"), *ReleaseVideoResourcesParams.PlayerGuid.ToString());
 					#endif
 
-					FExternalTextureRegistry::Get().UnregisterExternalTexture(Params.PlayerGuid);
+					FExternalTextureRegistry::Get().UnregisterExternalTexture(ReleaseVideoResourcesParams.PlayerGuid);
 
 					// @todo: this causes a crash
 //					Params.VideoTexture->Release();
@@ -669,10 +669,10 @@ void FAndroidMediaPlayer::TickFetch(FTimespan DeltaTime, FTimespan /*Timecode*/)
 			TSharedRef<FAndroidMediaTextureSample, ESPMode::ThreadSafe> VideoSample;
 			int32 SampleCount;
 		}
-		WriteVideoSampleParams = { JavaMediaPlayer, Samples, VideoSample, (int32)(VideoTrack.Dimensions.X * VideoTrack.Dimensions.Y * sizeof(int32)) };
+		Params = { JavaMediaPlayer, Samples, VideoSample, (int32)(VideoTrack.Dimensions.X * VideoTrack.Dimensions.Y * sizeof(int32)) };
 
-		ENQUEUE_UNIQUE_RENDER_COMMAND_ONEPARAMETER(AndroidMediaPlayerWriteVideoSample,
-			FWriteVideoSampleParams, Params, WriteVideoSampleParams,
+		ENQUEUE_RENDER_COMMAND(AndroidMediaPlayerWriteVideoSample)(
+			[Params](FRHICommandListImmediate& RHICmdList)
 			{
 				auto PinnedJavaMediaPlayer = Params.JavaMediaPlayerPtr.Pin();
 				auto PinnedSamples = Params.SamplesPtr.Pin();
@@ -714,10 +714,10 @@ void FAndroidMediaPlayer::TickFetch(FTimespan DeltaTime, FTimespan /*Timecode*/)
 			FGuid PlayerGuid;
 		};
 
-		FWriteVideoSampleParams WriteVideoSampleParams = { JavaMediaPlayer, PlayerGuid };
+		FWriteVideoSampleParams Params = { JavaMediaPlayer, PlayerGuid };
 
-		ENQUEUE_UNIQUE_RENDER_COMMAND_ONEPARAMETER(AndroidMediaPlayerWriteVideoSample,
-			FWriteVideoSampleParams, Params, WriteVideoSampleParams,
+		ENQUEUE_RENDER_COMMAND(AndroidMediaPlayerWriteVideoSample)(
+			[Params](FRHICommandListImmediate& RHICmdList)
 			{
 				if (IsRunningRHIInSeparateThread())
 				{
@@ -771,10 +771,10 @@ void FAndroidMediaPlayer::TickFetch(FTimespan DeltaTime, FTimespan /*Timecode*/)
 			TWeakPtr<FMediaSamples, ESPMode::ThreadSafe> SamplesPtr;
 			TSharedRef<FAndroidMediaTextureSample, ESPMode::ThreadSafe> VideoSample;
 		}
-		WriteVideoSampleParams = { JavaMediaPlayer, Samples, VideoSample };
+		Params = { JavaMediaPlayer, Samples, VideoSample };
 
-		ENQUEUE_UNIQUE_RENDER_COMMAND_ONEPARAMETER(AndroidMediaPlayerWriteVideoSample,
-			FWriteVideoSampleParams, Params, WriteVideoSampleParams,
+		ENQUEUE_RENDER_COMMAND(AndroidMediaPlayerWriteVideoSample)(
+			[Params](FRHICommandListImmediate& RHICmdList)
 			{
 				if (IsRunningRHIInSeparateThread())
 				{
diff --git a/Engine/Plugins/Runtime/GoogleVR/GoogleVRHMD/Source/GoogleVRHMD/Private/GoogleVRSplash.cpp b/Engine/Plugins/Runtime/GoogleVR/GoogleVRHMD/Source/GoogleVRHMD/Private/GoogleVRSplash.cpp
--- a/Engine/Plugins/Runtime/GoogleVR/GoogleVRHMD/Source/GoogleVRHMD/Private/GoogleVRSplash.cpp
+++ b/Engine/Plugins/Runtime/GoogleVR/GoogleVRHMD/Source/GoogleVRHMD/Private/GoogleVRSplash.cpp
@@ -134,14 +134,13 @@ void FGoogleVRSplash::Hide()
 	FTickableObjectRenderThread* RenderThreadTickerLocal = RenderThreadTicker.Get();
 	FGoogleVRSplash* pGVRSplash = this;
 	ENQUEUE_RENDER_COMMAND(UnregisterAsyncTick)(
-	TSharedPtr<FGoogleVRSplashTicker>&, RenderThreadTicker, RenderThreadTicker,
-	FGoogleVRSplash*, pGVRSplash, this,
-	{
-		pGVRSplash->SubmitBlackFrame();
+		[RenderThreadTickerLocal, pGVRSplash](FRHICommandListImmediate& RHICmdList)
+		{
+			pGVRSplash->SubmitBlackFrame();
 
-		RenderThreadTicker->Unregister();
-		RenderThreadTicker = nullptr;
-	});
+			RenderThreadTicker->Unregister();
+			RenderThreadTicker = nullptr;
+		});
 	FlushRenderingCommands();
 
 	if (!SplashTexturePath.IsEmpty())
