Enable webm for Android media player

#ue4
#android

diff --git a/Engine/Plugins/Media/AndroidMedia/Source/AndroidMediaEditor/Private/AndroidFileMediaSourceFactory.cpp b/Engine/Plugins/Media/AndroidMedia/Source/AndroidMediaEditor/Private/AndroidFileMediaSourceFactory.cpp
--- a/Engine/Plugins/Media/AndroidMedia/Source/AndroidMediaEditor/Private/AndroidFileMediaSourceFactory.cpp
+++ b/Engine/Plugins/Media/AndroidMedia/Source/AndroidMediaEditor/Private/AndroidFileMediaSourceFactory.cpp
@@ -13,6 +13,7 @@ UAndroidFileMediaSourceFactory::UAndroidFileMediaSourceFactory(const FObjectInit
 	Formats.Add(TEXT("3gpp;3GPP Multimedia File"));
 	Formats.Add(TEXT("aac;MPEG-2 Advanced Audio Coding File"));
 	Formats.Add(TEXT("mp4;MPEG-4 Movie"));
+	Formats.Add(TEXT("webm;WEBM Movie"));
 
 	SupportedClass = UFileMediaSource::StaticClass();
 	bEditorImport = true;
diff --git a/Engine/Plugins/Media/AndroidMedia/Source/AndroidMediaFactory/Private/AndroidMediaFactoryModule.cpp b/Engine/Plugins/Media/AndroidMedia/Source/AndroidMediaFactory/Private/AndroidMediaFactoryModule.cpp
--- a/Engine/Plugins/Media/AndroidMedia/Source/AndroidMediaFactory/Private/AndroidMediaFactoryModule.cpp
+++ b/Engine/Plugins/Media/AndroidMedia/Source/AndroidMediaFactory/Private/AndroidMediaFactoryModule.cpp
@@ -118,6 +118,7 @@ public:
 		SupportedFileExtensions.Add(TEXT("mp4"));
 		SupportedFileExtensions.Add(TEXT("m3u8"));
 		SupportedFileExtensions.Add(TEXT("mp3"));
+		SupportedFileExtensions.Add(TEXT("webm"));
 
 		// supported platforms
 		SupportedPlatforms.Add(TEXT("Android"));
