Fix up some misnamed functions on Android media plugins

#android

diff --git a/Engine/Plugins/Media/AndroidCamera/Source/AndroidCamera/Private/Player/AndroidCameraPlayer.cpp b/Engine/Plugins/Media/AndroidCamera/Source/AndroidCamera/Private/Player/AndroidCameraPlayer.cpp
--- a/Engine/Plugins/Media/AndroidCamera/Source/AndroidCamera/Private/Player/AndroidCameraPlayer.cpp
+++ b/Engine/Plugins/Media/AndroidCamera/Source/AndroidCamera/Private/Player/AndroidCameraPlayer.cpp
@@ -118,6 +118,20 @@ void FAndroidCameraPlayer::Close()
 	}
 
 	CurrentState = EMediaState::Closed;
+
+	// remove delegates if registered
+	if (ResumeHandle.IsValid())
+	{
+		FCoreDelegates::ApplicationHasEnteredForegroundDelegate.Remove(ResumeHandle);
+		ResumeHandle.Reset();
+	}
+
+	if (PauseHandle.IsValid())
+	{
+		FCoreDelegates::ApplicationWillEnterBackgroundDelegate.Remove(PauseHandle);
+		PauseHandle.Reset();
+	}
+
 	CurrentTime = FTimespan::Zero();
 
 	bLooping = false;
@@ -744,11 +758,11 @@ void FAndroidCameraPlayer::TickInput(FTimespan DeltaTime, FTimespan Timecode)
 	// register delegate if not registered
 	if (!ResumeHandle.IsValid())
 	{
-		ResumeHandle = FCoreDelegates::ApplicationHasEnteredForegroundDelegate.AddRaw(this, &FAndroidCameraPlayer::HandleApplicationWillEnterBackground);
+		ResumeHandle = FCoreDelegates::ApplicationHasEnteredForegroundDelegate.AddRaw(this, &FAndroidCameraPlayer::HandleApplicationHasEnteredForeground);
 	}
 	if (!PauseHandle.IsValid())
 	{
-		PauseHandle = FCoreDelegates::ApplicationWillEnterBackgroundDelegate.AddRaw(this, &FAndroidCameraPlayer::HandleApplicationHasEnteredForeground);
+		PauseHandle = FCoreDelegates::ApplicationWillEnterBackgroundDelegate.AddRaw(this, &FAndroidCameraPlayer::HandleApplicationWillEnterBackground);
 	}
 
 	// generate events
@@ -1502,7 +1516,7 @@ bool FAndroidCameraPlayer::SetRate(float Rate)
 /* FAndroidCameraPlayer callbacks
  *****************************************************************************/
 
-void FAndroidCameraPlayer::HandleApplicationHasEnteredForeground()
+void FAndroidCameraPlayer::HandleApplicationWillEnterBackground()
 {
 	/*
 	// check state in case changed before ticked
@@ -1514,7 +1528,7 @@ void FAndroidCameraPlayer::HandleApplicationHasEnteredForeground()
 }
 
 
-void FAndroidCameraPlayer::HandleApplicationWillEnterBackground()
+void FAndroidCameraPlayer::HandleApplicationHasEnteredForeground()
 {
 	/*
 	// check state in case changed before ticked
diff --git a/Engine/Plugins/Media/AndroidMedia/Source/AndroidMedia/Private/Player/AndroidMediaPlayer.cpp b/Engine/Plugins/Media/AndroidMedia/Source/AndroidMedia/Private/Player/AndroidMediaPlayer.cpp
--- a/Engine/Plugins/Media/AndroidMedia/Source/AndroidMedia/Private/Player/AndroidMediaPlayer.cpp
+++ b/Engine/Plugins/Media/AndroidMedia/Source/AndroidMedia/Private/Player/AndroidMediaPlayer.cpp
@@ -150,6 +150,19 @@ void FAndroidMediaPlayer::Close()
 
 	CurrentState = EMediaState::Closed;
 
+	// remove delegates if registered
+	if (ResumeHandle.IsValid())
+	{
+		FCoreDelegates::ApplicationHasEnteredForegroundDelegate.Remove(ResumeHandle);
+		ResumeHandle.Reset();
+	}
+ 
+	if (PauseHandle.IsValid())
+	{
+		FCoreDelegates::ApplicationWillEnterBackgroundDelegate.Remove(PauseHandle);
+		PauseHandle.Reset();
+	}
+
 	bLooping = false;
 
 	if (JavaMediaPlayer.IsValid())
@@ -906,11 +919,11 @@ void FAndroidMediaPlayer::TickInput(FTimespan DeltaTime, FTimespan /*Timecode*/)
 	// register delegate if not registered
 	if (!ResumeHandle.IsValid())
 	{
-		ResumeHandle = FCoreDelegates::ApplicationHasEnteredForegroundDelegate.AddRaw(this, &FAndroidMediaPlayer::HandleApplicationWillEnterBackground);
+		ResumeHandle = FCoreDelegates::ApplicationHasEnteredForegroundDelegate.AddRaw(this, &FAndroidMediaPlayer::HandleApplicationHasEnteredForeground);
 	}
 	if (!PauseHandle.IsValid())
 	{
-		PauseHandle = FCoreDelegates::ApplicationWillEnterBackgroundDelegate.AddRaw(this, &FAndroidMediaPlayer::HandleApplicationHasEnteredForeground);
+		PauseHandle = FCoreDelegates::ApplicationWillEnterBackgroundDelegate.AddRaw(this, &FAndroidMediaPlayer::HandleApplicationWillEnterBackground);
 	}
 
 	// generate events
@@ -1507,7 +1520,7 @@ bool FAndroidMediaPlayer::SetNativeVolume(float Volume)
 /* FAndroidMediaPlayer callbacks
  *****************************************************************************/
 
-void FAndroidMediaPlayer::HandleApplicationHasEnteredForeground()
+void FAndroidMediaPlayer::HandleApplicationWillEnterBackground()
 {
 	// check state in case changed before ticked
 	if ((CurrentState == EMediaState::Playing) && JavaMediaPlayer.IsValid())
@@ -1517,7 +1530,7 @@ void FAndroidMediaPlayer::HandleApplicationHasEnteredForeground()
 }
 
 
-void FAndroidMediaPlayer::HandleApplicationWillEnterBackground()
+void FAndroidMediaPlayer::HandleApplicationHasEnteredForeground()
 {
 	// check state in case changed before ticked
 	if ((CurrentState == EMediaState::Playing) && JavaMediaPlayer.IsValid())
