Fix a bug that media couldn't resume after entering background on Android devices.

diff --git a/Engine/Build/Android/Java/src/com/epicgames/unreal/MediaPlayer14.java b/Engine/Build/Android/Java/src/com/epicgames/unreal/MediaPlayer14.java
--- a/Engine/Build/Android/Java/src/com/epicgames/unreal/MediaPlayer14.java
+++ b/Engine/Build/Android/Java/src/com/epicgames/unreal/MediaPlayer14.java
@@ -638,6 +638,18 @@ public class MediaPlayer14
 		}
 	}
 
+	public void pause()
+	{
+		synchronized(this)
+		{
+			Completed = false;
+			if (Prepared)
+			{
+				super.pause();
+			}
+		}
+	}
+
 	public void stop()
 	{
 		synchronized(this)
diff --git a/Engine/Plugins/Media/AndroidMedia/Source/AndroidMedia/Private/Player/AndroidMediaPlayer.cpp b/Engine/Plugins/Media/AndroidMedia/Source/AndroidMedia/Private/Player/AndroidMediaPlayer.cpp
--- a/Engine/Plugins/Media/AndroidMedia/Source/AndroidMedia/Private/Player/AndroidMediaPlayer.cpp
+++ b/Engine/Plugins/Media/AndroidMedia/Source/AndroidMedia/Private/Player/AndroidMediaPlayer.cpp
@@ -47,6 +47,7 @@ FAndroidMediaPlayer::FAndroidMediaPlayer(IMediaEventSink& InEventSink)
 	, SelectedCaptionTrack(INDEX_NONE)
 	, SelectedVideoTrack(INDEX_NONE)
 	, VideoSamplePool(new FAndroidMediaTextureSamplePool)
+	, bEnterBackground(false)
 {
 	check(JavaMediaPlayer.IsValid());
 	check(Samples.IsValid());
@@ -930,7 +931,7 @@ void FAndroidMediaPlayer::TickInput(FTimespan DeltaTime, FTimespan /*Timecode*/)
 	}
 
 	// generate events
-	if (!JavaMediaPlayer->IsPlaying())
+	if (!JavaMediaPlayer->IsPlaying() && !bEnterBackground)
 	{
 		// might catch it restarting the loop so ignore if looping
 		if (!bLooping)
@@ -1528,8 +1529,27 @@ void FAndroidMediaPlayer::HandleApplicationWillEnterBackground()
 	// check state in case changed before ticked
 	if ((CurrentState == EMediaState::Playing) && JavaMediaPlayer.IsValid())
 	{
-		JavaMediaPlayer->Pause();
+		bool bDidComplete = JavaMediaPlayer->DidComplete();
+		
+		if (!bDidComplete || bLooping)
+		{
+			JavaMediaPlayer->Pause();
+		}
+		else
+		{
+			CurrentState = EMediaState::Stopped;
+		}
+
+		if (bDidComplete)
+		{
+			EventSink.ReceiveMediaEvent(EMediaEvent::PlaybackEndReached);
+
+#if ANDROIDMEDIAPLAYER_USE_NATIVELOGGING
+			FPlatformMisc::LowLevelOutputDebugStringf(TEXT("FAndroidMedia::Tick - PlaybackEndReached - stopped - %s"), *PlayerGuid.ToString());
+#endif
+		}
 	}
+	bEnterBackground = true;
 }
 
 
@@ -1540,6 +1560,7 @@ void FAndroidMediaPlayer::HandleApplicationHasEnteredForeground()
 	{
 		JavaMediaPlayer->Start();
 	}
+	bEnterBackground = false;
 }
 
 
diff --git a/Engine/Plugins/Media/AndroidMedia/Source/AndroidMedia/Private/Player/AndroidMediaPlayer.h b/Engine/Plugins/Media/AndroidMedia/Source/AndroidMedia/Private/Player/AndroidMediaPlayer.h
--- a/Engine/Plugins/Media/AndroidMedia/Source/AndroidMedia/Private/Player/AndroidMediaPlayer.h
+++ b/Engine/Plugins/Media/AndroidMedia/Source/AndroidMedia/Private/Player/AndroidMediaPlayer.h
@@ -167,4 +167,7 @@ private:
 
 	/** Video sample object pool. */
 	FAndroidMediaTextureSamplePool* VideoSamplePool;
+
+	/** Device enters background. */
+	bool bEnterBackground;
 };
