Fixed relative includes.

diff --git a/Engine/Plugins/Media/AndroidMedia/Source/AndroidMediaFactory/AndroidMediaFactory.Build.cs b/Engine/Plugins/Media/AndroidMedia/Source/AndroidMediaFactory/AndroidMediaFactory.Build.cs
--- a/Engine/Plugins/Media/AndroidMedia/Source/AndroidMediaFactory/AndroidMediaFactory.Build.cs
+++ b/Engine/Plugins/Media/AndroidMedia/Source/AndroidMediaFactory/AndroidMediaFactory.Build.cs
@@ -18,6 +18,7 @@ namespace UnrealBuildTool.Rules
 
 			PrivateIncludePathModuleNames.AddRange(
 				new string[] {
+					"AndroidMedia",
 					"Media",
 				});
 
diff --git a/Engine/Plugins/Media/AndroidMedia/Source/AndroidMediaFactory/Private/AndroidMediaFactoryModule.cpp b/Engine/Plugins/Media/AndroidMedia/Source/AndroidMediaFactory/Private/AndroidMediaFactoryModule.cpp
--- a/Engine/Plugins/Media/AndroidMedia/Source/AndroidMediaFactory/Private/AndroidMediaFactoryModule.cpp
+++ b/Engine/Plugins/Media/AndroidMedia/Source/AndroidMediaFactory/Private/AndroidMediaFactoryModule.cpp
@@ -9,7 +9,7 @@
 #include "IMediaModule.h"
 #include "IMediaPlayerFactory.h"
 
-#include "../../AndroidMedia/Public/IAndroidMediaModule.h"
+#include "IAndroidMediaModule.h"
 
 
 DEFINE_LOG_CATEGORY(LogAndroidMediaFactory);
diff --git a/Engine/Plugins/Media/AvfMedia/Source/AvfMediaFactory/AvfMediaFactory.Build.cs b/Engine/Plugins/Media/AvfMedia/Source/AvfMediaFactory/AvfMediaFactory.Build.cs
--- a/Engine/Plugins/Media/AvfMedia/Source/AvfMediaFactory/AvfMediaFactory.Build.cs
+++ b/Engine/Plugins/Media/AvfMedia/Source/AvfMediaFactory/AvfMediaFactory.Build.cs
@@ -18,6 +18,7 @@ namespace UnrealBuildTool.Rules
 
 			PrivateIncludePathModuleNames.AddRange(
 				new string[] {
+					"AvfMedia",
 					"Media",
 				});
 
diff --git a/Engine/Plugins/Media/AvfMedia/Source/AvfMediaFactory/Private/AvfMediaFactoryModule.cpp b/Engine/Plugins/Media/AvfMedia/Source/AvfMediaFactory/Private/AvfMediaFactoryModule.cpp
--- a/Engine/Plugins/Media/AvfMedia/Source/AvfMediaFactory/Private/AvfMediaFactoryModule.cpp
+++ b/Engine/Plugins/Media/AvfMedia/Source/AvfMediaFactory/Private/AvfMediaFactoryModule.cpp
@@ -21,7 +21,7 @@
 
 #include "AvfMediaSettings.h"
 
-#include "../../AvfMedia/Public/IAvfMediaModule.h"
+#include "IAvfMediaModule.h"
 
 
 DEFINE_LOG_CATEGORY(LogAvfMediaFactory);
diff --git a/Engine/Plugins/Media/ImgMedia/Source/ImgMediaFactory/Private/ImgMediaFactoryModule.cpp b/Engine/Plugins/Media/ImgMedia/Source/ImgMediaFactory/Private/ImgMediaFactoryModule.cpp
--- a/Engine/Plugins/Media/ImgMedia/Source/ImgMediaFactory/Private/ImgMediaFactoryModule.cpp
+++ b/Engine/Plugins/Media/ImgMedia/Source/ImgMediaFactory/Private/ImgMediaFactoryModule.cpp
@@ -22,7 +22,7 @@
 
 #include "ImgMediaSettings.h"
 
-#include "../../ImgMedia/Public/IImgMediaModule.h"
+#include "IImgMediaModule.h"
 
 
 DEFINE_LOG_CATEGORY(LogImgMediaFactory);
diff --git a/Engine/Plugins/Media/MfMedia/Source/MfMediaFactory/MfMediaFactory.Build.cs b/Engine/Plugins/Media/MfMedia/Source/MfMediaFactory/MfMediaFactory.Build.cs
--- a/Engine/Plugins/Media/MfMedia/Source/MfMediaFactory/MfMediaFactory.Build.cs
+++ b/Engine/Plugins/Media/MfMedia/Source/MfMediaFactory/MfMediaFactory.Build.cs
@@ -21,6 +21,7 @@ namespace UnrealBuildTool.Rules
 			PrivateIncludePathModuleNames.AddRange(
 				new string[] {
 					"Media",
+					"MfMedia",
 				});
 
 			if (DoAllowHTTPSPlayback())
diff --git a/Engine/Plugins/Media/MfMedia/Source/MfMediaFactory/Private/MfMediaFactoryModule.cpp b/Engine/Plugins/Media/MfMedia/Source/MfMediaFactory/Private/MfMediaFactoryModule.cpp
--- a/Engine/Plugins/Media/MfMedia/Source/MfMediaFactory/Private/MfMediaFactoryModule.cpp
+++ b/Engine/Plugins/Media/MfMedia/Source/MfMediaFactory/Private/MfMediaFactoryModule.cpp
@@ -10,7 +10,7 @@
 #include "Misc/Guid.h"
 #include "Modules/ModuleManager.h"
 
-#include "../../MfMedia/Public/IMfMediaModule.h"
+#include "IMfMediaModule.h"
 
 
 DEFINE_LOG_CATEGORY(LogMfMediaFactory);
diff --git a/Engine/Plugins/Media/WmfMedia/Source/WmfMediaFactory/Private/WmfMediaFactoryModule.cpp b/Engine/Plugins/Media/WmfMedia/Source/WmfMediaFactory/Private/WmfMediaFactoryModule.cpp
--- a/Engine/Plugins/Media/WmfMedia/Source/WmfMediaFactory/Private/WmfMediaFactoryModule.cpp
+++ b/Engine/Plugins/Media/WmfMedia/Source/WmfMediaFactory/Private/WmfMediaFactoryModule.cpp
@@ -20,7 +20,7 @@
 	#include "Misc/Guid.h"
 #endif
 
-#include "../../WmfMedia/Public/IWmfMediaModule.h"
+#include "IWmfMediaModule.h"
 
 
 DEFINE_LOG_CATEGORY(LogWmfMediaFactory);
diff --git a/Engine/Plugins/Media/WmfMedia/Source/WmfMediaFactory/WmfMediaFactory.Build.cs b/Engine/Plugins/Media/WmfMedia/Source/WmfMediaFactory/WmfMediaFactory.Build.cs
--- a/Engine/Plugins/Media/WmfMedia/Source/WmfMediaFactory/WmfMediaFactory.Build.cs
+++ b/Engine/Plugins/Media/WmfMedia/Source/WmfMediaFactory/WmfMediaFactory.Build.cs
@@ -19,6 +19,7 @@ namespace UnrealBuildTool.Rules
 			PrivateIncludePathModuleNames.AddRange(
 				new string[] {
 					"Media",
+					"WmfMedia",
 				});
 
 			PublicDependencyModuleNames.AddRange(
