Android GLES media player fixes.
Remove invalid call to GetNativeResource off the RHI thread which was causing assertion.
Restore the external texture path which was disabled accidentally in CL 30423346. In maintaining that fix, explicitely mark the AndroidMediaTextureSample as an external image and adjust MediaTextureResource's path to allow the conversion of external texture to standard 2D. Additionally, fix the channel swizzle in the external texture copy shader.

#rn Restore Android Media Player functionality under GLES.

diff --git a/Engine/Plugins/Media/AndroidMedia/Source/AndroidMedia/Private/Player/AndroidMediaPlayer.cpp b/Engine/Plugins/Media/AndroidMedia/Source/AndroidMedia/Private/Player/AndroidMediaPlayer.cpp
--- a/Engine/Plugins/Media/AndroidMedia/Source/AndroidMedia/Private/Player/AndroidMediaPlayer.cpp
+++ b/Engine/Plugins/Media/AndroidMedia/Source/AndroidMedia/Private/Player/AndroidMediaPlayer.cpp
@@ -492,7 +492,7 @@ static void ConditionalInitializeVideoTexture_RenderThread(TWeakPtr<FJavaAndroid
 		PinnedJavaMediaPlayer->SetVideoTextureValid(false);
 
 #if ANDROIDMEDIAPLAYER_USE_NATIVELOGGING
-		FPlatformMisc::LowLevelOutputDebugStringf(TEXT("Fetch RT: Created VideoTexture: %d - %s"), *reinterpret_cast<int32*>(VideoTexture->GetNativeResource()), *PlayerGuid.ToString());
+		FPlatformMisc::LowLevelOutputDebugStringf(TEXT("Fetch RT: Created VideoTexture for Player: %s"), *PlayerGuid.ToString());
 #endif
 	}
 }
@@ -763,7 +763,7 @@ void FAndroidMediaPlayer::TickFetch(FTimespan DeltaTime, FTimespan /*Timecode*/)
 
 		auto VideoSample = VideoSamplePool->AcquireShared();
 		const FJavaAndroidMediaPlayer::FVideoTrack VideoTrack = VideoTracks[SelectedVideoTrack];
-		if (!VideoSample->Initialize(VideoTrack.Dimensions, FTimespan::FromSeconds(1.0 / VideoTrack.FrameRate), FTimespan::FromMilliseconds(JavaMediaPlayer->GetCurrentPosition())))
+		if (!VideoSample->Initialize(VideoTrack.Dimensions, FTimespan::FromSeconds(1.0 / VideoTrack.FrameRate)))
 		{
 			return;
 		}
@@ -806,6 +806,10 @@ void FAndroidMediaPlayer::TickFetch(FTimespan DeltaTime, FTimespan /*Timecode*/)
 					PinnedJavaMediaPlayer->SetVideoTextureValid(true);
 				}
 
+				const int32 CurrentFramePosition = PinnedJavaMediaPlayer->GetCurrentPosition();
+				const FTimespan Time = FTimespan::FromMilliseconds(CurrentFramePosition);
+				
+				Params.VideoSample->InitializeExternalTexture(Time);
 				// Add a dummy video sample
 				// (actual data flows "on the side" via external texture mechanism)
 				PinnedSamples->AddVideo(Params.VideoSample);
diff --git a/Engine/Plugins/Media/AndroidMedia/Source/AndroidMedia/Private/Player/AndroidMediaTextureSample.h b/Engine/Plugins/Media/AndroidMedia/Source/AndroidMedia/Private/Player/AndroidMediaTextureSample.h
--- a/Engine/Plugins/Media/AndroidMedia/Source/AndroidMedia/Private/Player/AndroidMediaTextureSample.h
+++ b/Engine/Plugins/Media/AndroidMedia/Source/AndroidMedia/Private/Player/AndroidMediaTextureSample.h
@@ -31,6 +31,7 @@ public:
 	FAndroidMediaTextureSample()
 		: Buffer(nullptr)
 		, BufferSize(0)
+		, ExternalTexture(false)
 		, Dim(FIntPoint::ZeroValue)
 		, Duration(FTimespan::Zero())
 		, Time(FTimespan::Zero())
@@ -66,7 +67,7 @@ public:
 	 * @param InDuration The duration for which the sample is valid.
 	 * @return true on success, false otherwise.
 	 */
-	bool Initialize(const FIntPoint& InDim, FTimespan InDuration, FTimespan InTime = FTimespan::Zero())
+	bool Initialize(const FIntPoint& InDim, FTimespan InDuration)
 	{
 		if (InDim.GetMin() <= 0)
 		{
@@ -75,10 +76,20 @@ public:
 
 		Dim = InDim;
 		Duration = InDuration;
-		Time = InTime;
-
+		Time = FTimespan::Zero();
+		ExternalTexture = false;
+		
 		return true;
 	}
+	
+	/**
+	 * Initialize the sample as an external image reference.
+	 */
+	void InitializeExternalTexture(FTimespan InTime)
+	{
+		Time = InTime;
+		ExternalTexture = true;
+	}
 
 	/**
 	 * Initialize the sample with a memory buffer.
@@ -215,6 +226,11 @@ public:
 		return FMediaTimeStamp(Time);
 	}
 
+	virtual bool IsExternalImage() const override
+	{
+		return ExternalTexture;
+	}
+	
 	void UpdateTimeDuration(FTimespan InTime, FTimespan InDuration)
 	{
 		Time = InTime;
@@ -253,6 +269,9 @@ private:
 	/** Current allocation size of Buffer. */
 	SIZE_T BufferSize;
 
+	/** Flag indicating if this is an external image reference. */
+	bool ExternalTexture;
+	
 	/** Width and height of the texture sample. */
 	FIntPoint Dim;
 
diff --git a/Engine/Shaders/Private/MediaShaders.usf b/Engine/Shaders/Private/MediaShaders.usf
--- a/Engine/Shaders/Private/MediaShaders.usf
+++ b/Engine/Shaders/Private/MediaShaders.usf
@@ -550,7 +550,7 @@ void ReadTextureExternalPS(
 	out float4 OutColor : SV_Target0)
 {
 	float2 UV = float2(dot(InUV, ReadTextureExternalUB.ScaleRotation.xy), dot(InUV, ReadTextureExternalUB.ScaleRotation.zw)) + ReadTextureExternalUB.Offset;
-	float4 RGBOut = TextureExternalSample(InExternalTexture, ReadTextureExternalUB.SamplerP, UV).gbra;
+	float4 RGBOut = TextureExternalSample(InExternalTexture, ReadTextureExternalUB.SamplerP, UV);
 	OutColor = float4(sRGBToLinear(RGBOut.xyz), 1.0);
 }
 
diff --git a/Engine/Source/Runtime/Media/Public/IMediaTextureSample.h b/Engine/Source/Runtime/Media/Public/IMediaTextureSample.h
--- a/Engine/Source/Runtime/Media/Public/IMediaTextureSample.h
+++ b/Engine/Source/Runtime/Media/Public/IMediaTextureSample.h
@@ -363,6 +363,14 @@ public:
 		return (double)OutputDim.X / (double)OutputDim.Y;
 	}
 
+	/**
+	 * Indicate if this sample references an "external image"
+	 */
+	virtual bool IsExternalImage() const
+	{
+		return false;
+	}
+	
 	/**
 	 * Get the ScaleRotation (2x2 matrix) for the sample.
 	 *
diff --git a/Engine/Source/Runtime/MediaAssets/Private/Misc/MediaTextureResource.cpp b/Engine/Source/Runtime/MediaAssets/Private/Misc/MediaTextureResource.cpp
--- a/Engine/Source/Runtime/MediaAssets/Private/Misc/MediaTextureResource.cpp
+++ b/Engine/Source/Runtime/MediaAssets/Private/Misc/MediaTextureResource.cpp
@@ -345,11 +345,12 @@ FMediaTextureResource::FMediaTextureResource(UMediaTexture& InOwner, FIntPoint&
 	, CurrentMipMapBias(-1)
 	, PriorSamples(FPriorSamples::Create())
 {
+	bUsesImageExternal =
 #if PLATFORM_ANDROID
-	bUsesImageExternal = !Owner.NewStyleOutput && (!FAndroidMisc::ShouldUseVulkan() && GSupportsImageExternal);
-#else
-	bUsesImageExternal = !Owner.NewStyleOutput && GSupportsImageExternal;
+		!FAndroidMisc::ShouldUseVulkan() &&
 #endif
+		!Owner.NewStyleOutput && GSupportsImageExternal;
+	
 	if (OverrideColorSpaceType != UE::Color::EColorSpace::None)
 	{
 		OverrideColorSpace.Reset(new UE::Color::FColorSpace(OverrideColorSpaceType));
@@ -528,7 +529,7 @@ void FMediaTextureResource::Render(const FRenderParams& Params)
 				// No custom conversion, we need default processing...
 
 				// Sanity check: do we have any data?
-				if (Sample->GetTexture() == nullptr && Sample->GetBuffer() == nullptr)
+				if (!Sample->IsExternalImage() && Sample->GetTexture() == nullptr && Sample->GetBuffer() == nullptr)
 				{
 					UseSample = false;
 				}
