Fixed printfs.

diff --git a/Engine/Plugins/Media/AndroidCamera/Source/AndroidCamera/Private/Player/AndroidCameraPlayer.cpp b/Engine/Plugins/Media/AndroidCamera/Source/AndroidCamera/Private/Player/AndroidCameraPlayer.cpp
--- a/Engine/Plugins/Media/AndroidCamera/Source/AndroidCamera/Private/Player/AndroidCameraPlayer.cpp
+++ b/Engine/Plugins/Media/AndroidCamera/Source/AndroidCamera/Private/Player/AndroidCameraPlayer.cpp
@@ -1297,7 +1297,7 @@ bool FAndroidCameraPlayer::SelectTrack(EMediaTrackType TrackType, int32 TrackInd
 
 			if (TrackIndex == INDEX_NONE)
 			{
-				UE_LOG(LogAndroidCamera, VeryVerbose, TEXT("Player %p: Disabling audio"), this, TrackIndex);
+				UE_LOG(LogAndroidCamera, VeryVerbose, TEXT("Player %p: Disabling audio"), this);
 
 				JavaCameraPlayer->SetAudioEnabled(false);
 			}
@@ -1308,7 +1308,7 @@ bool FAndroidCameraPlayer::SelectTrack(EMediaTrackType TrackType, int32 TrackInd
 					return false;
 				}
 
-				UE_LOG(LogAndroidCamera, VeryVerbose, TEXT("Player %p: Enabling audio"), this, TrackIndex);
+				UE_LOG(LogAndroidCamera, VeryVerbose, TEXT("Player %p: Enabling audio"), this);
 
 				JavaCameraPlayer->SetAudioEnabled(true);
 			}
@@ -1324,7 +1324,7 @@ bool FAndroidCameraPlayer::SelectTrack(EMediaTrackType TrackType, int32 TrackInd
 
 			if (TrackIndex == INDEX_NONE)
 			{
-				UE_LOG(LogAndroidCamera, VeryVerbose, TEXT("Player %p: Disabling captions"), this, TrackIndex);
+				UE_LOG(LogAndroidCamera, VeryVerbose, TEXT("Player %p: Disabling captions"), this);
 			}
 			else
 			{
@@ -1333,7 +1333,7 @@ bool FAndroidCameraPlayer::SelectTrack(EMediaTrackType TrackType, int32 TrackInd
 					return false;
 				}
 
-				UE_LOG(LogAndroidCamera, VeryVerbose, TEXT("Player %p: Enabling captions"), this, TrackIndex);
+				UE_LOG(LogAndroidCamera, VeryVerbose, TEXT("Player %p: Enabling captions"), this);
 			}
 
 			SelectedCaptionTrack = TrackIndex;
@@ -1347,7 +1347,7 @@ bool FAndroidCameraPlayer::SelectTrack(EMediaTrackType TrackType, int32 TrackInd
 
 			if (TrackIndex == INDEX_NONE)
 			{
-				UE_LOG(LogAndroidCamera, VeryVerbose, TEXT("Player %p: Disabling video"), this, TrackIndex);
+				UE_LOG(LogAndroidCamera, VeryVerbose, TEXT("Player %p: Disabling video"), this);
 				JavaCameraPlayer->SetVideoEnabled(false);
 			}
 			else
diff --git a/Engine/Plugins/Media/AndroidMedia/Source/AndroidMedia/Private/Player/AndroidMediaPlayer.cpp b/Engine/Plugins/Media/AndroidMedia/Source/AndroidMedia/Private/Player/AndroidMediaPlayer.cpp
--- a/Engine/Plugins/Media/AndroidMedia/Source/AndroidMedia/Private/Player/AndroidMediaPlayer.cpp
+++ b/Engine/Plugins/Media/AndroidMedia/Source/AndroidMedia/Private/Player/AndroidMediaPlayer.cpp
@@ -1349,7 +1349,7 @@ bool FAndroidMediaPlayer::SelectTrack(EMediaTrackType TrackType, int32 TrackInde
 
 			if (TrackIndex == INDEX_NONE)
 			{
-				UE_LOG(LogAndroidMedia, VeryVerbose, TEXT("Player %p: Disabling audio"), this, TrackIndex);
+				UE_LOG(LogAndroidMedia, VeryVerbose, TEXT("Player %p: Disabling audio"), this);
 
 				JavaMediaPlayer->SetAudioEnabled(false);
 			}
@@ -1360,7 +1360,7 @@ bool FAndroidMediaPlayer::SelectTrack(EMediaTrackType TrackType, int32 TrackInde
 					return false;
 				}
 
-				UE_LOG(LogAndroidMedia, VeryVerbose, TEXT("Player %p: Enabling audio"), this, TrackIndex);
+				UE_LOG(LogAndroidMedia, VeryVerbose, TEXT("Player %p: Enabling audio"), this);
 
 				JavaMediaPlayer->SetAudioEnabled(true);
 			}
@@ -1376,7 +1376,7 @@ bool FAndroidMediaPlayer::SelectTrack(EMediaTrackType TrackType, int32 TrackInde
 
 			if (TrackIndex == INDEX_NONE)
 			{
-				UE_LOG(LogAndroidMedia, VeryVerbose, TEXT("Player %p: Disabling captions"), this, TrackIndex);
+				UE_LOG(LogAndroidMedia, VeryVerbose, TEXT("Player %p: Disabling captions"), this);
 			}
 			else
 			{
@@ -1385,7 +1385,7 @@ bool FAndroidMediaPlayer::SelectTrack(EMediaTrackType TrackType, int32 TrackInde
 					return false;
 				}
 
-				UE_LOG(LogAndroidMedia, VeryVerbose, TEXT("Player %p: Enabling captions"), this, TrackIndex);
+				UE_LOG(LogAndroidMedia, VeryVerbose, TEXT("Player %p: Enabling captions"), this);
 			}
 
 			SelectedCaptionTrack = TrackIndex;
@@ -1399,7 +1399,7 @@ bool FAndroidMediaPlayer::SelectTrack(EMediaTrackType TrackType, int32 TrackInde
 
 			if (TrackIndex == INDEX_NONE)
 			{
-				UE_LOG(LogAndroidMedia, VeryVerbose, TEXT("Player %p: Disabling video"), this, TrackIndex);
+				UE_LOG(LogAndroidMedia, VeryVerbose, TEXT("Player %p: Disabling video"), this);
 				JavaMediaPlayer->SetVideoEnabled(false);
 			}
 			else
@@ -1409,7 +1409,7 @@ bool FAndroidMediaPlayer::SelectTrack(EMediaTrackType TrackType, int32 TrackInde
 					return false;
 				}
 
-				UE_LOG(LogAndroidMedia, VeryVerbose, TEXT("Player %p: Enabling video"), this, TrackIndex);
+				UE_LOG(LogAndroidMedia, VeryVerbose, TEXT("Player %p: Enabling video"), this);
 				JavaMediaPlayer->SetVideoEnabled(true);
 			}
 
diff --git a/Engine/Plugins/Media/AvfMedia/Source/AvfMedia/Private/Player/AvfMediaPlayer.cpp b/Engine/Plugins/Media/AvfMedia/Source/AvfMedia/Private/Player/AvfMediaPlayer.cpp
--- a/Engine/Plugins/Media/AvfMedia/Source/AvfMedia/Private/Player/AvfMediaPlayer.cpp
+++ b/Engine/Plugins/Media/AvfMedia/Source/AvfMedia/Private/Player/AvfMediaPlayer.cpp
@@ -696,7 +696,7 @@ bool FAvfMediaPlayer::Open(const FString& Url, const IMediaOptions* /*Options*/)
 	// open media file
 	if (nsMediaUrl == nil)
 	{
-		UE_LOG(LogAvfMedia, Error, TEXT("Failed to open Media file:"), *Url);
+		UE_LOG(LogAvfMedia, Error, TEXT("Failed to open Media file: %s"), *Url);
 
 		return false;
 	}
@@ -742,7 +742,7 @@ bool FAvfMediaPlayer::Open(const FString& Url, const IMediaOptions* /*Options*/)
 
 	if (PlayerItem == nil)
 	{
-		UE_LOG(LogAvfMedia, Error, TEXT("Failed to open player item with Url:"), *Url);
+		UE_LOG(LogAvfMedia, Error, TEXT("Failed to open player item with Url: %s"), *Url);
 		return false;
 	}
 
diff --git a/Engine/Source/Editor/UnrealEd/Private/SourceCodeNavigation.cpp b/Engine/Source/Editor/UnrealEd/Private/SourceCodeNavigation.cpp
--- a/Engine/Source/Editor/UnrealEd/Private/SourceCodeNavigation.cpp
+++ b/Engine/Source/Editor/UnrealEd/Private/SourceCodeNavigation.cpp
@@ -700,7 +700,7 @@ void FSourceCodeNavigationImpl::NavigateToFunctionSource( const FString& Functio
                 {
                     if (!FPaths::FileExists(FullDsymPath))
                     {
-                        UE_LOG(LogSelectionDetails, Warning, TEXT("NavigateToFunctionSource: Missing dSYM files, please install Editor symbols for debugging"), *Results);
+                        UE_LOG(LogSelectionDetails, Warning, TEXT("NavigateToFunctionSource: Missing dSYM files, please install Editor symbols for debugging"));
                     }
                     else
                     {
diff --git a/Engine/Source/Runtime/Online/WebSockets/Private/WinHttp/Support/WinHttpConnectionWebSocket.cpp b/Engine/Source/Runtime/Online/WebSockets/Private/WinHttp/Support/WinHttpConnectionWebSocket.cpp
--- a/Engine/Source/Runtime/Online/WebSockets/Private/WinHttp/Support/WinHttpConnectionWebSocket.cpp
+++ b/Engine/Source/Runtime/Online/WebSockets/Private/WinHttp/Support/WinHttpConnectionWebSocket.cpp
@@ -810,7 +810,7 @@ void FWinHttpConnectionWebSocket::HandleWebSocketStatusCallback(HINTERNET Resour
 			// This shouldn't happen, but ignoring this case since we don't handle it below
 			if (SocketStatus->eBufferType == WINHTTP_WEB_SOCKET_CLOSE_BUFFER_TYPE)
 			{
-				UE_LOG(LogWinHttp, Error, TEXT("WinHttp WebSocket sent unexpected Close message"), this);
+				UE_LOG(LogWinHttp, Error, TEXT("WinHttp WebSocket sent unexpected Close message"));
 				return;
 			}
 
